Motion Design: renamed scene state, data link plugin friendly names

diff --git a/Engine/Plugins/Experimental/AvalancheDataLink/AvalancheDataLink.uplugin b/Engine/Plugins/Experimental/AvalancheDataLink/AvalancheDataLink.uplugin
--- a/Engine/Plugins/Experimental/AvalancheDataLink/AvalancheDataLink.uplugin
+++ b/Engine/Plugins/Experimental/AvalancheDataLink/AvalancheDataLink.uplugin
@@ -2,7 +2,7 @@
 	"FileVersion": 3,
 	"Version": 1,
 	"VersionName": "1.0",
-	"FriendlyName": "DataLink for Motion Design",
+	"FriendlyName": "Motion Design Data Link Integration",
 	"Description": "",
 	"Category": "Other",
 	"CreatedBy": "",
diff --git a/Engine/Plugins/Experimental/DataLink/DataLink.uplugin b/Engine/Plugins/Experimental/DataLink/DataLink.uplugin
--- a/Engine/Plugins/Experimental/DataLink/DataLink.uplugin
+++ b/Engine/Plugins/Experimental/DataLink/DataLink.uplugin
@@ -2,7 +2,7 @@
 	"FileVersion": 3,
 	"Version": 1,
 	"VersionName": "1.0",
-	"FriendlyName": "DataLink",
+	"FriendlyName": "Motion Design Data Link",
 	"Description": "",
 	"Category": "Other",
 	"CreatedBy": "",
diff --git a/Engine/Plugins/Experimental/DataLinkOAuth/DataLinkOAuth.uplugin b/Engine/Plugins/Experimental/DataLinkOAuth/DataLinkOAuth.uplugin
--- a/Engine/Plugins/Experimental/DataLinkOAuth/DataLinkOAuth.uplugin
+++ b/Engine/Plugins/Experimental/DataLinkOAuth/DataLinkOAuth.uplugin
@@ -2,7 +2,7 @@
 	"FileVersion": 3,
 	"Version": 1,
 	"VersionName": "1.0",
-	"FriendlyName": "DataLink OAuth",
+	"FriendlyName": "Motion Design Data Link OAuth",
 	"Description": "Data Link functionality for OAuth 2.0",
 	"Category": "Other",
 	"CreatedBy": "",
diff --git a/Engine/Plugins/Experimental/SceneState/SceneState.uplugin b/Engine/Plugins/Experimental/SceneState/SceneState.uplugin
--- a/Engine/Plugins/Experimental/SceneState/SceneState.uplugin
+++ b/Engine/Plugins/Experimental/SceneState/SceneState.uplugin
@@ -2,7 +2,7 @@
 	"FileVersion": 3,
 	"Version": 1,
 	"VersionName": "1.0",
-	"FriendlyName": "Scene State",
+	"FriendlyName": "Motion Design Scene State",
 	"Description": "",
 	"Category": "Experimental",
 	"CreatedBy": "",
