Motion Design: added promotion copy flow when copying between property types that are different but that can be copied via promotion.

diff --git a/Engine/Plugins/Experimental/AvalancheDataLink/Source/AvalancheDataLink/AvalancheDataLink.Build.cs b/Engine/Plugins/Experimental/AvalancheDataLink/Source/AvalancheDataLink/AvalancheDataLink.Build.cs
--- a/Engine/Plugins/Experimental/AvalancheDataLink/Source/AvalancheDataLink/AvalancheDataLink.Build.cs
+++ b/Engine/Plugins/Experimental/AvalancheDataLink/Source/AvalancheDataLink/AvalancheDataLink.Build.cs
@@ -26,6 +26,7 @@ public class AvalancheDataLink : ModuleRules
                 "Engine",
                 "Json",
                 "JsonUtilities",
+                "PropertyBindingUtils",
                 "RemoteControl",
                 "RemoteControlLogic",
             }
diff --git a/Engine/Plugins/Experimental/AvalancheDataLink/Source/AvalancheDataLink/Private/AvaDataLinkInstance.cpp b/Engine/Plugins/Experimental/AvalancheDataLink/Source/AvalancheDataLink/Private/AvaDataLinkInstance.cpp
--- a/Engine/Plugins/Experimental/AvalancheDataLink/Source/AvalancheDataLink/Private/AvaDataLinkInstance.cpp
+++ b/Engine/Plugins/Experimental/AvalancheDataLink/Source/AvalancheDataLink/Private/AvaDataLinkInstance.cpp
@@ -14,12 +14,90 @@
 #include "Engine/World.h"
 #include "JsonObjectConverter.h"
 #include "JsonObjectWrapper.h"
+#include "PropertyBindingTypes.h"
 #include "RCVirtualProperty.h"
 #include "RCVirtualPropertyContainer.h"
 #include "RemoteControlPreset.h"
 #include "StructUtils/StructView.h"
 #include "UObject/PropertyAccessUtil.h"
 
+namespace UE::AvaDataLink::Private
+{
+	struct FCopyInfo
+	{
+		const FProperty* SourceProperty;
+		const uint8* SourceMemory;
+		FProperty* TargetProperty;
+		uint8* TargetMemory;
+	};
+
+	/**
+	 * Tries promoting a given source property value to match the target property type and copies it to the target memory
+	 * Derived from FPropertyBindingBindingCollection::PerformCopy
+	 */
+	template<typename InSourcePropertyType>
+	class FPromotionCopy
+	{
+		using FSourceCppType = typename InSourcePropertyType::TCppType;
+
+		template<typename InTargetPropertyType>
+		bool CopySingle(const FCopyInfo& InCopyInfo, const InSourcePropertyType& InSourceProperty)
+		{
+			using FTargetCppType = typename InTargetPropertyType::TCppType;
+
+			static_assert(std::is_convertible_v<FSourceCppType, FTargetCppType>);
+
+			if (InTargetPropertyType* TargetProperty = CastField<InTargetPropertyType>(InCopyInfo.TargetProperty))
+			{
+				TargetProperty->SetPropertyValue(InCopyInfo.TargetMemory, static_cast<FTargetCppType>(InSourceProperty.GetPropertyValue(InCopyInfo.SourceMemory)));
+				return true;
+			}
+			return false;
+		}
+
+		template<typename>
+		bool CopyInternal(const FCopyInfo& InCopyInfo, const InSourcePropertyType& InSourceProperty)
+		{
+			return false;
+		}
+
+		template<typename, typename InTargetPropertyType, typename ...InOtherTargetPropertyTypes>
+		bool CopyInternal(const FCopyInfo& InCopyInfo, const InSourcePropertyType& InSourceProperty)
+		{
+			return this->CopySingle<InTargetPropertyType>(InCopyInfo, InSourceProperty) || this->CopyInternal<void, InOtherTargetPropertyTypes...>(InCopyInfo, InSourceProperty);
+		}
+
+	public:
+		/** Tries promoting from a source property type to any of the given target property types */
+		template<typename ...InTargetPropertyTypes>
+		bool Copy(const FCopyInfo& InCopyInfo)
+		{
+			if (const InSourcePropertyType* SourceProperty = CastField<InSourcePropertyType>(InCopyInfo.SourceProperty))
+			{
+				return this->CopyInternal<void, InTargetPropertyTypes...>(InCopyInfo, *SourceProperty);
+			}
+			return false;
+		}
+	};
+
+	bool PromoteCopy(const FCopyInfo& InCopyInfo)
+	{
+		// Bool Promotions
+		return FPromotionCopy<FBoolProperty>().Copy<FByteProperty, FIntProperty, FUInt32Property, FInt64Property, FFloatProperty, FDoubleProperty>(InCopyInfo)
+			// Byte Promotion
+			|| FPromotionCopy<FByteProperty>().Copy<FIntProperty, FUInt32Property, FInt64Property, FFloatProperty, FDoubleProperty>(InCopyInfo)
+			// Int32 Promotions
+			|| FPromotionCopy<FIntProperty>().Copy<FInt64Property, FFloatProperty, FDoubleProperty>(InCopyInfo)
+			// Uint32 Promotions
+			|| FPromotionCopy<FUInt32Property>().Copy<FInt64Property, FFloatProperty, FDoubleProperty>(InCopyInfo)
+			// Float Promotions
+			|| FPromotionCopy<FFloatProperty>().Copy<FIntProperty, FInt64Property, FDoubleProperty>(InCopyInfo)
+			// Double Promotions
+			|| FPromotionCopy<FDoubleProperty>().Copy<FIntProperty, FInt64Property, FFloatProperty>(InCopyInfo)
+		;
+	}
+} // UE::AvaDataLink::Private
+
 void UAvaDataLinkInstance::Execute()
 {
 	if (Executor.IsValid())
@@ -107,22 +185,46 @@ void UAvaDataLinkInstance::OnExecutionFinished(const FDataLinkExecutor& InExecut
 	else
 	{
 		ForEachResolvedController(InExecutor, Preset, TargetDataView,
-			[&InOutputDataView, &ModifiedControllers](const FResolvedController& InResolvedController)
+			[&InOutputDataView, &ModifiedControllers, &InExecutor](const FResolvedController& InResolvedController)
 			{
 				const FProperty* SourceProperty = PropertyAccessUtil::FindPropertyByName(*InResolvedController.Mapping->OutputFieldName, InOutputDataView.GetScriptStruct());
 				if (!SourceProperty)
 				{
+					UE_LOG(LogAvaDataLink, Warning, TEXT("[%s] Data Link output field name '%s' could not be applied as it was not found in output data struct '%s'")
+						, InExecutor.GetContextName().GetData()
+						, *InResolvedController.Mapping->OutputFieldName
+						, *GetNameSafe(InOutputDataView.GetScriptStruct()));
 					return;
 				}
 
-				if (!SourceProperty->SameType(InResolvedController.TargetProperty))
+				const UE::PropertyBinding::EPropertyCompatibility Compatibility = UE::PropertyBinding::GetPropertyCompatibility(SourceProperty, InResolvedController.TargetProperty);
+				if (Compatibility == UE::PropertyBinding::EPropertyCompatibility::Incompatible)
 				{
+					UE_LOG(LogAvaDataLink, Warning, TEXT("[%s] Data Link output '%s' could not be applied to controller '%s' as types are incompatible")
+						, InExecutor.GetContextName().GetData()
+						, *InResolvedController.Mapping->OutputFieldName
+						, *InResolvedController.Mapping->TargetController.Name.ToString());
 					return;
 				}
 
 				const uint8* SourceMemory = SourceProperty->ContainerPtrToValuePtr<uint8>(InOutputDataView.GetMemory());
-				InResolvedController.TargetProperty->CopyCompleteValue(InResolvedController.TargetMemory, SourceMemory);
-				ModifiedControllers.AddUnique(InResolvedController.Controller);
+				switch (Compatibility)
+				{
+				case UE::PropertyBinding::EPropertyCompatibility::Compatible:
+					InResolvedController.TargetProperty->CopyCompleteValue(InResolvedController.TargetMemory, SourceMemory);
+					ModifiedControllers.AddUnique(InResolvedController.Controller);
+					break;
+
+				case UE::PropertyBinding::EPropertyCompatibility::Promotable:
+					ensureMsgf(UE::AvaDataLink::Private::PromoteCopy({
+						.SourceProperty = SourceProperty,
+						.SourceMemory = SourceMemory,
+						.TargetProperty = InResolvedController.TargetProperty,
+						.TargetMemory = InResolvedController.TargetMemory
+					}), TEXT("Promotion failed even though compatibility was deemed as 'promotable'."));
+					ModifiedControllers.AddUnique(InResolvedController.Controller);
+					break;
+				}
 			});
 	}
 
diff --git a/Engine/Plugins/Experimental/AvalancheSceneState/Source/AvalancheSceneState/AvalancheSceneState.Build.cs b/Engine/Plugins/Experimental/AvalancheSceneState/Source/AvalancheSceneState/AvalancheSceneState.Build.cs
--- a/Engine/Plugins/Experimental/AvalancheSceneState/Source/AvalancheSceneState/AvalancheSceneState.Build.cs
+++ b/Engine/Plugins/Experimental/AvalancheSceneState/Source/AvalancheSceneState/AvalancheSceneState.Build.cs
@@ -27,6 +27,7 @@ public class AvalancheSceneState : ModuleRules
 			{
 				"Avalanche",
 				"Engine",
+				"PropertyBindingUtils",
 				"RemoteControl",
 				"RemoteControlLogic",
 				"SceneStateEvent",
diff --git a/Engine/Plugins/Experimental/AvalancheSceneState/Source/AvalancheSceneState/Private/RemoteControl/AvaSceneStateRCTask.cpp b/Engine/Plugins/Experimental/AvalancheSceneState/Source/AvalancheSceneState/Private/RemoteControl/AvaSceneStateRCTask.cpp
--- a/Engine/Plugins/Experimental/AvalancheSceneState/Source/AvalancheSceneState/Private/RemoteControl/AvaSceneStateRCTask.cpp
+++ b/Engine/Plugins/Experimental/AvalancheSceneState/Source/AvalancheSceneState/Private/RemoteControl/AvaSceneStateRCTask.cpp
@@ -12,6 +12,85 @@
 #include "SceneStateExecutionContext.h"
 #include "UObject/PropertyAccessUtil.h"
 
+namespace UE::AvaSceneState::Private
+{
+	struct FCopyInfo
+	{
+		const FProperty* SourceProperty;
+		const uint8* SourceMemory;
+		FProperty* TargetProperty;
+		uint8* TargetMemory;
+	};
+
+	/**
+	 * Tries promoting a given source property value to match the target property type and copies it to the target memory
+	 * Derived from FPropertyBindingBindingCollection::PerformCopy
+	 * Todo: merge with FPromotionCopy in Motion Design Data Link Integration
+	 */
+	template<typename InSourcePropertyType>
+	class FPromotionCopy
+	{
+		using FSourceCppType = typename InSourcePropertyType::TCppType;
+
+		template<typename InTargetPropertyType>
+		bool CopySingle(const FCopyInfo& InCopyInfo, const InSourcePropertyType& InSourceProperty)
+		{
+			using FTargetCppType = typename InTargetPropertyType::TCppType;
+
+			static_assert(std::is_convertible_v<FSourceCppType, FTargetCppType>);
+
+			if (InTargetPropertyType* TargetProperty = CastField<InTargetPropertyType>(InCopyInfo.TargetProperty))
+			{
+				TargetProperty->SetPropertyValue(InCopyInfo.TargetMemory, static_cast<FTargetCppType>(InSourceProperty.GetPropertyValue(InCopyInfo.SourceMemory)));
+				return true;
+			}
+			return false;
+		}
+
+		template<typename>
+		bool CopyInternal(const FCopyInfo& InCopyInfo, const InSourcePropertyType& InSourceProperty)
+		{
+			return false;
+		}
+
+		template<typename, typename InTargetPropertyType, typename ...InOtherTargetPropertyTypes>
+		bool CopyInternal(const FCopyInfo& InCopyInfo, const InSourcePropertyType& InSourceProperty)
+		{
+			return this->CopySingle<InTargetPropertyType>(InCopyInfo, InSourceProperty) || this->CopyInternal<void, InOtherTargetPropertyTypes...>(InCopyInfo, InSourceProperty);
+		}
+
+	public:
+		/** Tries promoting from a source property type to any of the given target property types */
+		template<typename ...InTargetPropertyTypes>
+		bool Copy(const FCopyInfo& InCopyInfo)
+		{
+			if (const InSourcePropertyType* SourceProperty = CastField<InSourcePropertyType>(InCopyInfo.SourceProperty))
+			{
+				return this->CopyInternal<void, InTargetPropertyTypes...>(InCopyInfo, *SourceProperty);
+			}
+			return false;
+		}
+	};
+
+	bool PromoteCopy(const FCopyInfo& InCopyInfo)
+	{
+		// Bool Promotions
+		return FPromotionCopy<FBoolProperty>().Copy<FByteProperty, FIntProperty, FUInt32Property, FInt64Property, FFloatProperty, FDoubleProperty>(InCopyInfo)
+			// Byte Promotion
+			|| FPromotionCopy<FByteProperty>().Copy<FIntProperty, FUInt32Property, FInt64Property, FFloatProperty, FDoubleProperty>(InCopyInfo)
+			// Int32 Promotions
+			|| FPromotionCopy<FIntProperty>().Copy<FInt64Property, FFloatProperty, FDoubleProperty>(InCopyInfo)
+			// Uint32 Promotions
+			|| FPromotionCopy<FUInt32Property>().Copy<FInt64Property, FFloatProperty, FDoubleProperty>(InCopyInfo)
+			// Float Promotions
+			|| FPromotionCopy<FFloatProperty>().Copy<FIntProperty, FInt64Property, FDoubleProperty>(InCopyInfo)
+			// Double Promotions
+			|| FPromotionCopy<FDoubleProperty>().Copy<FIntProperty, FInt64Property, FFloatProperty>(InCopyInfo)
+		;
+	}
+
+} // UE::AvaSceneState::Private
+
 FAvaSceneStateRCTask::FAvaSceneStateRCTask()
 {
 	SetFlags(ESceneStateTaskFlags::HasBindingExtension);
@@ -97,20 +176,45 @@ void FAvaSceneStateRCTask::OnStart(const FSceneStateExecutionContext& InContext,
 		}
 
 		FProperty* const TargetProperty = Controller->GetProperty();
-		if (!ensure(TargetProperty) || !TargetProperty->SameType(SourcePropertyDesc.CachedProperty))
+		if (!TargetProperty)
 		{
-			UE_LOG(LogAvaSceneState, Warning, TEXT("[%s] Controller '%s' type mismatch in preset '%s'")
+			UE_LOG(LogAvaSceneState, Warning, TEXT("[%s] Controller '%s' has invalid target property in preset '%s'")
 				, *InContext.GetExecutionContextName()
 				, *Controller->DisplayName.ToString()
 				, *Preset->GetName());
 			continue;
 		}
 
-		const void* SourceData = SourcePropertyDesc.CachedProperty->ContainerPtrToValuePtr<void>(SourceValuesMemory);
-		void* TargetData = TargetProperty->ContainerPtrToValuePtr<void>(TargetControllerDataView.GetMemory());
+		const UE::PropertyBinding::EPropertyCompatibility Compatibility = UE::PropertyBinding::GetPropertyCompatibility(SourcePropertyDesc.CachedProperty, TargetProperty);
+		if (Compatibility == UE::PropertyBinding::EPropertyCompatibility::Incompatible)
+		{
+			UE_LOG(LogAvaSceneState, Warning, TEXT("[%s] Controller '%s' has type incompatibility in preset '%s'")
+				, *InContext.GetExecutionContextName()
+				, *Controller->DisplayName.ToString()
+				, *Preset->GetName());
+			continue;
+		}
 
-		TargetProperty->CopyCompleteValue(TargetData, SourceData);
-		ModifiedControllers.AddUnique(Controller);
+		const uint8* SourceMemory = SourcePropertyDesc.CachedProperty->ContainerPtrToValuePtr<uint8>(SourceValuesMemory);
+		uint8* TargetMemory = TargetProperty->ContainerPtrToValuePtr<uint8>(TargetControllerDataView.GetMemory());
+
+		switch (Compatibility)
+		{
+		case UE::PropertyBinding::EPropertyCompatibility::Compatible:
+			TargetProperty->CopyCompleteValue(TargetMemory, SourceMemory);
+			ModifiedControllers.AddUnique(Controller);
+			break;
+
+		case UE::PropertyBinding::EPropertyCompatibility::Promotable:
+			ensureMsgf(UE::AvaSceneState::Private::PromoteCopy({
+				.SourceProperty = SourcePropertyDesc.CachedProperty,
+				.SourceMemory = SourceMemory,
+				.TargetProperty = TargetProperty,
+				.TargetMemory = TargetMemory
+			}), TEXT("Promotion failed even though compatibility was deemed as 'promotable'."));
+			ModifiedControllers.AddUnique(Controller);
+			break;
+		}
 	}
 
 	TSet<FGuid> ModifiedControllerIds;
