Motion Design Data Link: added sub-property mapping support for json and controller mappings

diff --git a/Engine/Plugins/Experimental/AvalancheDataLink/Source/AvalancheDataLink/Private/AvaDataLinkInstance.cpp b/Engine/Plugins/Experimental/AvalancheDataLink/Source/AvalancheDataLink/Private/AvaDataLinkInstance.cpp
--- a/Engine/Plugins/Experimental/AvalancheDataLink/Source/AvalancheDataLink/Private/AvaDataLinkInstance.cpp
+++ b/Engine/Plugins/Experimental/AvalancheDataLink/Source/AvalancheDataLink/Private/AvaDataLinkInstance.cpp
@@ -174,17 +174,18 @@ void UAvaDataLinkInstance::OnExecutionFinished(const FDataLinkExecutor& InExecut
 		ForEachResolvedController(InExecutor, Preset, TargetDataView,
 			[&InOutputDataView, &ModifiedControllers, &InExecutor](const FResolvedController& InResolvedController)
 			{
-				const FProperty* SourceProperty = PropertyAccessUtil::FindPropertyByName(*InResolvedController.Mapping->OutputFieldName, InOutputDataView.GetScriptStruct());
-				if (!SourceProperty)
+				FString ErrorMessage;
+				const UE::DataLink::FConstPropertyView SourceView = UE::DataLink::ResolveConstPropertyView(InOutputDataView, InResolvedController.Mapping->OutputFieldName, &ErrorMessage);
+				if (!SourceView.Property)
 				{
-					UE_LOG(LogAvaDataLink, Warning, TEXT("[%s] Data Link output field name '%s' could not be applied as it was not found in output data struct '%s'")
+					UE_LOG(LogAvaDataLink, Warning, TEXT("[%s] Data Link output field name '%s' could not be applied: %s")
 						, InExecutor.GetContextName().GetData()
 						, *InResolvedController.Mapping->OutputFieldName
-						, *GetNameSafe(InOutputDataView.GetScriptStruct()));
+						, *ErrorMessage);
 					return;
 				}
 
-				const UE::PropertyBinding::EPropertyCompatibility Compatibility = UE::PropertyBinding::GetPropertyCompatibility(SourceProperty, InResolvedController.TargetProperty);
+				const UE::PropertyBinding::EPropertyCompatibility Compatibility = UE::PropertyBinding::GetPropertyCompatibility(SourceView.Property, InResolvedController.TargetProperty);
 				if (Compatibility == UE::PropertyBinding::EPropertyCompatibility::Incompatible)
 				{
 					UE_LOG(LogAvaDataLink, Warning, TEXT("[%s] Data Link output '%s' could not be applied to controller '%s' as types are incompatible")
@@ -194,18 +195,17 @@ void UAvaDataLinkInstance::OnExecutionFinished(const FDataLinkExecutor& InExecut
 					return;
 				}
 
-				const uint8* SourceMemory = SourceProperty->ContainerPtrToValuePtr<uint8>(InOutputDataView.GetMemory());
 				switch (Compatibility)
 				{
 				case UE::PropertyBinding::EPropertyCompatibility::Compatible:
-					InResolvedController.TargetProperty->CopyCompleteValue(InResolvedController.TargetMemory, SourceMemory);
+					InResolvedController.TargetProperty->CopyCompleteValue(InResolvedController.TargetMemory, SourceView.Memory);
 					ModifiedControllers.AddUnique(InResolvedController.Controller);
 					break;
 
 				case UE::PropertyBinding::EPropertyCompatibility::Promotable:
 					ensureMsgf(UE::AvaDataLink::Private::PromoteCopy({
-						.SourceProperty = SourceProperty,
-						.SourceMemory = SourceMemory,
+						.SourceProperty = SourceView.Property,
+						.SourceMemory = SourceView.Memory,
 						.TargetProperty = InResolvedController.TargetProperty,
 						.TargetMemory = InResolvedController.TargetMemory
 					}), TEXT("Promotion failed even though compatibility was deemed as 'promotable'."));
@@ -273,8 +273,8 @@ void UAvaDataLinkInstance::ForEachResolvedController(const FDataLinkExecutor& In
 			continue;
 		}
 
-		ResolvedMapping.TargetProperty = PropertyAccessUtil::FindPropertyByName(ResolvedMapping.Controller->PropertyName, InTargetDataView.GetScriptStruct());
-		if (!ResolvedMapping.TargetProperty)
+		const UE::DataLink::FPropertyView TargetPropertyView = UE::DataLink::ResolvePropertyView(InTargetDataView, ResolvedMapping.Controller->PropertyName.ToString());
+		if (!TargetPropertyView.Property)
 		{
 			UE_LOG(LogAvaDataLink, Warning, TEXT("[%s] Data Link output '%s' could not be applied to controller '%s'. Controller property '%s' in preset '%s' was not found!")
 				, InExecutor.GetContextName().GetData()
@@ -285,7 +285,8 @@ void UAvaDataLinkInstance::ForEachResolvedController(const FDataLinkExecutor& In
 			continue;
 		}
 
-		ResolvedMapping.TargetMemory = ResolvedMapping.TargetProperty->ContainerPtrToValuePtr<uint8>(InTargetDataView.GetMemory());
+		ResolvedMapping.TargetProperty = TargetPropertyView.Property;
+		ResolvedMapping.TargetMemory = TargetPropertyView.Memory;
 		InFunction(ResolvedMapping);
 	}
 }
diff --git a/Engine/Plugins/Experimental/DataLink/Source/DataLink/Private/DataLinkUtils.cpp b/Engine/Plugins/Experimental/DataLink/Source/DataLink/Private/DataLinkUtils.cpp
--- a/Engine/Plugins/Experimental/DataLink/Source/DataLink/Private/DataLinkUtils.cpp
+++ b/Engine/Plugins/Experimental/DataLink/Source/DataLink/Private/DataLinkUtils.cpp
@@ -9,9 +9,11 @@
 #include "DataLinkSinkObject.h"
 #include "Engine/Engine.h"
 #include "IDataLinkSinkProvider.h"
+#include "StructUtils/StructView.h"
 #include "UObject/Class.h"
 #include "UObject/Object.h"
 #include "UObject/Package.h"
+#include "UObject/PropertyAccessUtil.h"
 
 namespace UE::DataLink
 {
@@ -90,6 +92,85 @@ TSharedPtr<FDataLinkSink> TryGetSink(TScriptInterface<IDataLinkSinkProvider> InS
 	return nullptr;
 }
 
+FConstPropertyView ResolveConstPropertyView(FConstStructView InBaseStructView, const FString& InPropertyPath, FString* OutError)
+{
+	if (!InBaseStructView.IsValid())
+	{
+		if (OutError)
+		{
+			*OutError = FString::Printf(TEXT("FindPropertyByPath Error. Invalid base struct data"));
+		}
+		return FConstPropertyView();
+	}
+
+	TArray<FString> PathSegments;
+	InPropertyPath.ParseIntoArray(PathSegments, TEXT("."), /*bInCullEmpty*/true);
+	if (PathSegments.IsEmpty())
+	{
+		if (OutError)
+		{
+			*OutError = FString::Printf(TEXT("FindPropertyByPath Error. Path '%s' is not properly formatted. No path segments (defined by delimiter '.') found."), *InPropertyPath);
+		}
+		return FConstPropertyView();
+	}
+
+	const UStruct* CurrentStruct = InBaseStructView.GetScriptStruct();
+
+	FConstPropertyView Current;
+	Current.Memory = InBaseStructView.GetMemory();
+
+	for (const FString& PathSegment : PathSegments)
+	{
+		if (!CurrentStruct)
+		{
+			if (OutError)
+			{
+				*OutError = FString::Printf(TEXT("FindPropertyByPath Error. Path could not find property '%s' as previous property was not a struct. Path: %s")
+					, *PathSegment
+					, *InPropertyPath);
+			}
+			return FConstPropertyView();
+		}
+
+		Current.Property = PropertyAccessUtil::FindPropertyByName(*PathSegment, CurrentStruct);
+		if (!Current.Property)
+		{
+			if (OutError)
+			{
+				*OutError = FString::Printf(TEXT("FindPropertyByPath Error. Property '%s' could not be found in struct '%s'. Path: %s.")
+					, *PathSegment
+					, *CurrentStruct->GetName()
+					, *InPropertyPath);
+			}
+			return FConstPropertyView();
+		}
+
+		Current.Memory = Current.Property->ContainerPtrToValuePtr<uint8>(Current.Memory);
+
+		if (const FStructProperty* StructProperty = CastField<FStructProperty>(Current.Property))
+		{
+			CurrentStruct = StructProperty->Struct;
+		}
+		else
+		{
+			CurrentStruct = nullptr;
+		}
+	}
+
+	return Current;
+}
+
+FPropertyView ResolvePropertyView(FStructView InBaseStructView, const FString& InPropertyPath, FString* OutError)
+{
+	FPropertyView Result;
+
+	const FConstPropertyView ResolvedProperty = ResolveConstPropertyView(InBaseStructView, InPropertyPath, OutError);
+	Result.Property = const_cast<FProperty*>(ResolvedProperty.Property);
+	Result.Memory = const_cast<uint8*>(ResolvedProperty.Memory);
+
+	return Result;
+}
+
 void SetInputData(UDataLinkGraph* InGraph, TArray<FDataLinkInputData>& OutInputData)
 {
 	if (!InGraph)
diff --git a/Engine/Plugins/Experimental/DataLink/Source/DataLink/Public/DataLinkUtils.h b/Engine/Plugins/Experimental/DataLink/Source/DataLink/Public/DataLinkUtils.h
--- a/Engine/Plugins/Experimental/DataLink/Source/DataLink/Public/DataLinkUtils.h
+++ b/Engine/Plugins/Experimental/DataLink/Source/DataLink/Public/DataLinkUtils.h
@@ -3,15 +3,20 @@
 #pragma once
 
 #include "Containers/ContainersFwd.h"
+#include "Containers/StringFwd.h"
 #include "HAL/Platform.h"
 #include "Templates/SharedPointerFwd.h"
 
+class FProperty;
 class IDataLinkSinkProvider;
 class UClass;
 class UDataLinkGraph;
 class UObject;
+class UStruct;
+struct FConstStructView;
 struct FDataLinkInputData;
 struct FDataLinkSink;
+struct FStructView;
 template <typename T> class TScriptInterface;
 
 #define UE_API DATALINK_API
@@ -34,6 +39,34 @@ namespace UE::DataLink
 	 */
 	UE_API TSharedPtr<FDataLinkSink> TryGetSink(TScriptInterface<IDataLinkSinkProvider> InSinkProvider);
 
+	struct FConstPropertyView
+	{
+		const FProperty* Property = nullptr;
+		const uint8* Memory = nullptr;
+	};
+	/**
+	 * Sets the input data to match the current graph's inputs
+	 * @param InBaseStructView the base struct to start looking for the property path
+	 * @param InPropertyPath the property path to look for
+	 * @param OutError optional error message if property returns null
+	 * @return the resolved property view (property and memory ptr)
+	 */
+	UE_API FConstPropertyView ResolveConstPropertyView(FConstStructView InBaseStructView, const FString& InPropertyPath, FString* OutError = nullptr);
+
+	struct FPropertyView
+	{
+		FProperty* Property = nullptr;
+		uint8* Memory = nullptr;
+	};
+	/**
+	 * Sets the input data to match the current graph's inputs
+	 * @param InBaseStructView the base struct to start looking for the property path
+	 * @param InPropertyPath the property path to look for
+	 * @param OutError optional error message if property returns null
+	 * @return the resolved property view (property and memory ptr)
+	 */
+	UE_API FPropertyView ResolvePropertyView(FStructView InBaseStructView, const FString& InPropertyPath, FString* OutError = nullptr);
+
 	/**
 	 * Sets the input data to match the current graph's inputs
 	 * @param InGraph the data link graph to get the input pins to update from
diff --git a/Engine/Plugins/Experimental/DataLink/Source/DataLinkJson/Private/DataLinkJsonStructSimpleMapping.cpp b/Engine/Plugins/Experimental/DataLink/Source/DataLinkJson/Private/DataLinkJsonStructSimpleMapping.cpp
--- a/Engine/Plugins/Experimental/DataLink/Source/DataLinkJson/Private/DataLinkJsonStructSimpleMapping.cpp
+++ b/Engine/Plugins/Experimental/DataLink/Source/DataLinkJson/Private/DataLinkJsonStructSimpleMapping.cpp
@@ -3,24 +3,21 @@
 #include "DataLinkJsonStructSimpleMapping.h"
 #include "DataLinkJsonLog.h"
 #include "DataLinkJsonUtils.h"
+#include "DataLinkUtils.h"
 #include "JsonObjectConverter.h"
 #include "StructUtils/StructView.h"
-#include "UObject/PropertyAccessUtil.h"
 
 bool UDataLinkJsonStructSimpleMapping::Apply(const TSharedRef<FJsonObject>& InSourceJson, const FStructView& InTargetStructView) const
 {
-	const UScriptStruct* const TargetStruct = InTargetStructView.GetScriptStruct();
-
-	uint8* const TargetMemory = InTargetStructView.GetMemory();
-
 	for (const TPair<FString, FName>& FieldMapping : FieldMappings)
 	{
-		FProperty* const TargetProperty = PropertyAccessUtil::FindPropertyByName(FieldMapping.Value, TargetStruct);
-		if (!TargetProperty)
+		FString ErrorMessage;
+		const UE::DataLink::FPropertyView TargetPropertyView = UE::DataLink::ResolvePropertyView(InTargetStructView, FieldMapping.Value.ToString(), &ErrorMessage);
+		if (!TargetPropertyView.Property)
 		{
-			UE_LOG(LogDataLinkJson, Error, TEXT("Field '%s' not found in struct '%s'")
+			UE_LOG(LogDataLinkJson, Error, TEXT("Field '%s' not found. %s")
 				, *FieldMapping.Value.ToString()
-				, *TargetStruct->GetName());
+				, *ErrorMessage);
 			continue;
 		}
 
@@ -32,13 +29,12 @@ bool UDataLinkJsonStructSimpleMapping::Apply(const TSharedRef<FJsonObject>& InSo
 			continue;
 		}
 
-		uint8* const TargetPropertyValue = TargetProperty->ContainerPtrToValuePtr<uint8>(TargetMemory);
-		if (!FJsonObjectConverter::JsonValueToUProperty(JsonValue, TargetProperty, TargetPropertyValue))
+		if (!FJsonObjectConverter::JsonValueToUProperty(JsonValue, TargetPropertyView.Property, TargetPropertyView.Memory))
 		{
 			UE_LOG(LogDataLinkJson, Error, TEXT("Could not copy Json Value with key '%s' to property '%s' in struct '%s'")
 				, *FieldMapping.Key
-				, *TargetProperty->GetName()
-				, *TargetStruct->GetName());
+				, *TargetPropertyView.Property->GetName()
+				, *GetNameSafe(InTargetStructView.GetScriptStruct()));
 		}
 	}
 
