Motion Design Data Link: implemented output processors and refactored motion design data link instance to use these processors for RC

diff --git a/Engine/Plugins/Experimental/AvalancheDataLink/Source/AvalancheDataLink/Private/AvaDataLinkActor.cpp b/Engine/Plugins/Experimental/AvalancheDataLink/Source/AvalancheDataLink/Private/AvaDataLinkActor.cpp
--- a/Engine/Plugins/Experimental/AvalancheDataLink/Source/AvalancheDataLink/Private/AvaDataLinkActor.cpp
+++ b/Engine/Plugins/Experimental/AvalancheDataLink/Source/AvalancheDataLink/Private/AvaDataLinkActor.cpp
@@ -24,3 +24,10 @@ void AAvaDataLinkActor::StopDataLinkInstances()
 		}
 	}
 }
+
+#if WITH_EDITOR
+FString AAvaDataLinkActor::GetDefaultActorLabel() const
+{
+	return TEXT("Motion Design Data Link Actor");
+}
+#endif
diff --git a/Engine/Plugins/Experimental/AvalancheDataLink/Source/AvalancheDataLink/Private/AvaDataLinkActor.h b/Engine/Plugins/Experimental/AvalancheDataLink/Source/AvalancheDataLink/Private/AvaDataLinkActor.h
--- a/Engine/Plugins/Experimental/AvalancheDataLink/Source/AvalancheDataLink/Private/AvaDataLinkActor.h
+++ b/Engine/Plugins/Experimental/AvalancheDataLink/Source/AvalancheDataLink/Private/AvaDataLinkActor.h
@@ -19,6 +19,12 @@ public:
 	UFUNCTION(BlueprintCallable, CallInEditor, Category="Motion Design Data Link")
 	void StopDataLinkInstances();
 
+	//~ Begin AActor
+#if WITH_EDITOR
+	virtual FString GetDefaultActorLabel() const override;
+#endif
+	//~ End AActor
+
 	UPROPERTY(EditAnywhere, Instanced, Category="Motion Design Data Link")
 	TArray<TObjectPtr<UAvaDataLinkInstance>> DataLinkInstances;
 };
diff --git a/Engine/Plugins/Experimental/AvalancheDataLink/Source/AvalancheDataLink/Private/AvaDataLinkInstance.cpp b/Engine/Plugins/Experimental/AvalancheDataLink/Source/AvalancheDataLink/Private/AvaDataLinkInstance.cpp
--- a/Engine/Plugins/Experimental/AvalancheDataLink/Source/AvalancheDataLink/Private/AvaDataLinkInstance.cpp
+++ b/Engine/Plugins/Experimental/AvalancheDataLink/Source/AvalancheDataLink/Private/AvaDataLinkInstance.cpp
@@ -1,89 +1,10 @@
 // Copyright Epic Games, Inc. All Rights Reserved.
 
 #include "AvaDataLinkInstance.h"
-#include "AvaDataLinkLog.h"
-#include "AvaSceneSubsystem.h"
-#include "Avalanche/Public/IAvaSceneInterface.h"
-#include "Controller/RCController.h"
-#include "DataLinkEnums.h"
 #include "DataLinkExecutor.h"
 #include "DataLinkExecutorArguments.h"
-#include "DataLinkJsonUtils.h"
-#include "DataLinkUtils.h"
-#include "Engine/Level.h"
 #include "Engine/World.h"
-#include "JsonObjectConverter.h"
-#include "JsonObjectWrapper.h"
-#include "PropertyBindingTypes.h"
-#include "RCVirtualProperty.h"
-#include "RCVirtualPropertyContainer.h"
-#include "RemoteControlPreset.h"
-#include "StructUtils/StructView.h"
-#include "UObject/PropertyAccessUtil.h"
-
-namespace UE::AvaDataLink::Private
-{
-	struct FCopyInfo
-	{
-		const FProperty* SourceProperty;
-		const uint8* SourceMemory;
-		FProperty* TargetProperty;
-		uint8* TargetMemory;
-	};
-
-	/**
-	 * Tries promoting a given source property value to match the target property type and copies it to the target memory
-	 * Derived from FPropertyBindingBindingCollection::PerformCopy
-	 */
-	template<typename InSourcePropertyType>
-	class FPromotionCopy
-	{
-		template<typename InTargetPropertyType>
-		bool CopySingle(const FCopyInfo& InCopyInfo, const InSourcePropertyType& InSourceProperty)
-		{
-			using FSourceCppType = InSourcePropertyType::TCppType;
-			using FTargetCppType = InTargetPropertyType::TCppType;
-
-			static_assert(std::is_convertible_v<FSourceCppType, FTargetCppType>);
-
-			if (InTargetPropertyType* TargetProperty = CastField<InTargetPropertyType>(InCopyInfo.TargetProperty))
-			{
-				TargetProperty->SetPropertyValue(InCopyInfo.TargetMemory, InSourceProperty.GetPropertyValue(InCopyInfo.SourceMemory));
-				return true;
-			}
-			return false;
-		}
-
-	public:
-		/** Tries promoting from a source property type to any of the given target property types */
-		template<typename ...InTargetPropertyTypes>
-		bool Copy(const FCopyInfo& InCopyInfo)
-		{
-			if (const InSourcePropertyType* SourceProperty = CastField<InSourcePropertyType>(InCopyInfo.SourceProperty))
-			{
-				return (this->CopySingle<InTargetPropertyTypes>(InCopyInfo, *SourceProperty) || ...);
-			}
-			return false;
-		}
-	};
-
-	bool PromoteCopy(const FCopyInfo& InCopyInfo)
-	{
-		// Bool Promotions
-		return FPromotionCopy<FBoolProperty>().Copy<FByteProperty, FIntProperty, FUInt32Property, FInt64Property, FFloatProperty, FDoubleProperty>(InCopyInfo)
-			// Byte Promotion
-			|| FPromotionCopy<FByteProperty>().Copy<FIntProperty, FUInt32Property, FInt64Property, FFloatProperty, FDoubleProperty>(InCopyInfo)
-			// Int32 Promotions
-			|| FPromotionCopy<FIntProperty>().Copy<FInt64Property, FFloatProperty, FDoubleProperty>(InCopyInfo)
-			// Uint32 Promotions
-			|| FPromotionCopy<FUInt32Property>().Copy<FInt64Property, FFloatProperty, FDoubleProperty>(InCopyInfo)
-			// Float Promotions
-			|| FPromotionCopy<FFloatProperty>().Copy<FIntProperty, FInt64Property, FDoubleProperty>(InCopyInfo)
-			// Double Promotions
-			|| FPromotionCopy<FDoubleProperty>().Copy<FIntProperty, FInt64Property, FFloatProperty>(InCopyInfo)
-		;
-	}
-} // UE::AvaDataLink::Private
+#include "RemoteControl/AvaDataLinkRCProcessor.h"
 
 void UAvaDataLinkInstance::Execute()
 {
@@ -94,7 +15,7 @@ void UAvaDataLinkInstance::Execute()
 		.SetContextName(BuildContextName())
 #endif
 		.SetContextObject(this)
-		.SetOnOutputData(FOnDataLinkOutputData::CreateUObject(this, &UAvaDataLinkInstance::OnOutputData)));
+		.SetOutputProcessors(OutputProcessors));
 
 	Executor->Run();
 }
@@ -108,188 +29,33 @@ void UAvaDataLinkInstance::Stop()
 	}
 }
 
-#if WITH_DATALINK_CONTEXT
-FString UAvaDataLinkInstance::BuildContextName() const
+#if WITH_EDITORONLY_DATA
+void UAvaDataLinkInstance::DeclareConstructClasses(TArray<FTopLevelAssetPath>& OutConstructClasses, const UClass* InSpecificSubclass)
 {
-	return FString::Printf(TEXT("Motion Design Data Link. World: '%s'"), *GetNameSafe(GetTypedOuter<UWorld>()));
+	Super::DeclareConstructClasses(OutConstructClasses, InSpecificSubclass);
+	OutConstructClasses.Add(FTopLevelAssetPath(UAvaDataLinkRCProcessor::StaticClass()));
 }
 #endif
 
-void UAvaDataLinkInstance::OnOutputData(const FDataLinkExecutor& InExecutor, FConstStructView InOutputDataView)
+void UAvaDataLinkInstance::PostLoad()
 {
-	if (!InOutputDataView.IsValid())
-	{
-		return;
-	}
+	Super::PostLoad();
 
-	URemoteControlPreset* const Preset = GetRemoteControlPreset();
-	if (!Preset)
+#if WITH_EDITOR
+PRAGMA_DISABLE_DEPRECATION_WARNINGS
+	if (!ControllerMappings_DEPRECATED.IsEmpty() && OutputProcessors.IsEmpty())
 	{
-		UE_LOG(LogAvaDataLink, Error, TEXT("[%s] Data Link execution finished, but Remote Control is invalid!"), InExecutor.GetContextName().GetData());
-		return;
+		UAvaDataLinkRCProcessor* const RCProcessor = NewObject<UAvaDataLinkRCProcessor>(this, NAME_None, GetMaskedFlags(RF_PropagateToSubObjects));
+		RCProcessor->ControllerMappings = ControllerMappings_DEPRECATED;
+		OutputProcessors.Add(RCProcessor);
 	}
-
-	URCVirtualPropertyContainerBase* const ControllerContainer = Preset->GetControllerContainer();
-	if (!ControllerContainer)
-	{
-		UE_LOG(LogAvaDataLink, Error, TEXT("[%s] Data Link execution finished, but Remote Control '%s' has invalid Controller Container!")
-			, InExecutor.GetContextName().GetData()
-			, *Preset->GetName());
-		return;
-	}
-
-	const FStructView TargetDataView = ControllerContainer->GetPropertyBagMutableValue();
-	if (!TargetDataView.IsValid())
-	{
-		UE_LOG(LogAvaDataLink, Error, TEXT("[%s] Data Link execution finished, but Remote Control '%s' has invalid Controller Container Data View!")
-			, InExecutor.GetContextName().GetData()
-			, *Preset->GetName());
-		return;
-	}
-
-	TArray<URCController*> ModifiedControllers;
-	ModifiedControllers.Reserve(ControllerMappings.Num());
-
-	if (const FJsonObjectWrapper* OutputJson = InOutputDataView.GetPtr<const FJsonObjectWrapper>())
-	{
-		if (!OutputJson->JsonObject.IsValid())
-		{
-			UE_LOG(LogAvaDataLink, Warning, TEXT("[%s] Data Link output could not be applied to controllers. Json Object was not valid!"), InExecutor.GetContextName().GetData());
-			return;
-		}
-
-		TSharedRef<FJsonObject> JsonObject = OutputJson->JsonObject.ToSharedRef();
-
-		ForEachResolvedController(InExecutor, Preset, TargetDataView,
-			[&JsonObject, &ModifiedControllers](const FResolvedController& InResolvedController)
-			{
-				TSharedPtr<FJsonValue> SourceJsonValue = UE::DataLinkJson::FindJsonValue(JsonObject, InResolvedController.Mapping->OutputFieldName);
-
-				// If property was set, return true to signal that the controller value was modified
-				if (FJsonObjectConverter::JsonValueToUProperty(SourceJsonValue, InResolvedController.TargetProperty, InResolvedController.TargetMemory))
-				{
-					ModifiedControllers.AddUnique(InResolvedController.Controller);
-				}
-			});
-	}
-	else
-	{
-		ForEachResolvedController(InExecutor, Preset, TargetDataView,
-			[&InOutputDataView, &ModifiedControllers, &InExecutor](const FResolvedController& InResolvedController)
-			{
-				FString ErrorMessage;
-				const UE::DataLink::FConstPropertyView SourceView = UE::DataLink::ResolveConstPropertyView(InOutputDataView, InResolvedController.Mapping->OutputFieldName, &ErrorMessage);
-				if (!SourceView.Property)
-				{
-					UE_LOG(LogAvaDataLink, Warning, TEXT("[%s] Data Link output field name '%s' could not be applied: %s")
-						, InExecutor.GetContextName().GetData()
-						, *InResolvedController.Mapping->OutputFieldName
-						, *ErrorMessage);
-					return;
-				}
-
-				const UE::PropertyBinding::EPropertyCompatibility Compatibility = UE::PropertyBinding::GetPropertyCompatibility(SourceView.Property, InResolvedController.TargetProperty);
-				if (Compatibility == UE::PropertyBinding::EPropertyCompatibility::Incompatible)
-				{
-					UE_LOG(LogAvaDataLink, Warning, TEXT("[%s] Data Link output '%s' could not be applied to controller '%s' as types are incompatible")
-						, InExecutor.GetContextName().GetData()
-						, *InResolvedController.Mapping->OutputFieldName
-						, *InResolvedController.Mapping->TargetController.Name.ToString());
-					return;
-				}
-
-				switch (Compatibility)
-				{
-				case UE::PropertyBinding::EPropertyCompatibility::Compatible:
-					InResolvedController.TargetProperty->CopyCompleteValue(InResolvedController.TargetMemory, SourceView.Memory);
-					ModifiedControllers.AddUnique(InResolvedController.Controller);
-					break;
-
-				case UE::PropertyBinding::EPropertyCompatibility::Promotable:
-					ensureMsgf(UE::AvaDataLink::Private::PromoteCopy({
-						.SourceProperty = SourceView.Property,
-						.SourceMemory = SourceView.Memory,
-						.TargetProperty = InResolvedController.TargetProperty,
-						.TargetMemory = InResolvedController.TargetMemory
-					}), TEXT("Promotion failed even though compatibility was deemed as 'promotable'."));
-					ModifiedControllers.AddUnique(InResolvedController.Controller);
-					break;
-				}
-			});
-	}
-
-	TSet<FGuid> ModifiedControllerIds;
-	ModifiedControllerIds.Reserve(ModifiedControllers.Num());
-
-	for (URCController* Controller : ModifiedControllers)
-	{
-		Controller->OnModifyPropertyValue();
-		ModifiedControllerIds.Add(Controller->Id);
-	}
-
-	Preset->OnControllerModified().Broadcast(Preset, ModifiedControllerIds);
-}
-
-IAvaSceneInterface* UAvaDataLinkInstance::GetSceneInterface() const
-{
-	ULevel* Level = GetTypedOuter<ULevel>();
-	if (!Level || !Level->OwningWorld)
-	{
-		return nullptr;
-	}
-
-	const UAvaSceneSubsystem* SceneSubsystem = Level->OwningWorld->GetSubsystem<UAvaSceneSubsystem>();
-	if (!SceneSubsystem)
-	{
-		// No subsystem yet available, try finding the scene interface by iterating the actors
-		return UAvaSceneSubsystem::FindSceneInterface(Level);
-	}
-
-	return SceneSubsystem->GetSceneInterface(Level);
-}
-
-URemoteControlPreset* UAvaDataLinkInstance::GetRemoteControlPreset() const
-{
-	if (IAvaSceneInterface* SceneInterface = GetSceneInterface())
-	{
-		return SceneInterface->GetRemoteControlPreset();
-	}
-	return nullptr;
+PRAGMA_ENABLE_DEPRECATION_WARNINGS
+#endif
 }
 
-void UAvaDataLinkInstance::ForEachResolvedController(const FDataLinkExecutor& InExecutor, URemoteControlPreset* InPreset, FStructView InTargetDataView, TFunctionRef<void(const FResolvedController&)> InFunction)
+#if WITH_DATALINK_CONTEXT
+FString UAvaDataLinkInstance::BuildContextName() const
 {
-	for (const FAvaDataLinkControllerMapping& Mapping : ControllerMappings)
-	{
-		FResolvedController ResolvedMapping;
-		ResolvedMapping.Mapping = &Mapping;
-
-		// Get the controller to retrieve the underlying Property FName of the controller within the Controller Container Property Bag
-		ResolvedMapping.Controller = Cast<URCController>(Mapping.TargetController.FindController(InPreset));
-		if (!ResolvedMapping.Controller)
-		{
-			UE_LOG(LogAvaDataLink, Warning, TEXT("[%s] Data Link output '%s' could not be applied to controller '%s'. Controller was not found in preset '%s'!")
-				, InExecutor.GetContextName().GetData()
-				, *Mapping.OutputFieldName
-				, *Mapping.TargetController.Name.ToString()
-				, *InPreset->GetName());
-			continue;
-		}
-
-		const UE::DataLink::FPropertyView TargetPropertyView = UE::DataLink::ResolvePropertyView(InTargetDataView, ResolvedMapping.Controller->PropertyName.ToString());
-		if (!TargetPropertyView.Property)
-		{
-			UE_LOG(LogAvaDataLink, Warning, TEXT("[%s] Data Link output '%s' could not be applied to controller '%s'. Controller property '%s' in preset '%s' was not found!")
-				, InExecutor.GetContextName().GetData()
-				, *Mapping.OutputFieldName
-				, *Mapping.TargetController.Name.ToString()
-				, *ResolvedMapping.Controller->PropertyName.ToString()
-				, *InPreset->GetName());
-			continue;
-		}
-
-		ResolvedMapping.TargetProperty = TargetPropertyView.Property;
-		ResolvedMapping.TargetMemory = TargetPropertyView.Memory;
-		InFunction(ResolvedMapping);
-	}
+	return FString::Printf(TEXT("Motion Design Data Link. World: '%s'"), *GetNameSafe(GetTypedOuter<UWorld>()));
 }
+#endif
diff --git a/Engine/Plugins/Experimental/AvalancheDataLink/Source/AvalancheDataLink/Private/AvaDataLinkProcessor.cpp b/Engine/Plugins/Experimental/AvalancheDataLink/Source/AvalancheDataLink/Private/AvaDataLinkProcessor.cpp
new file mode 100644
--- /dev/null
+++ b/Engine/Plugins/Experimental/AvalancheDataLink/Source/AvalancheDataLink/Private/AvaDataLinkProcessor.cpp
@@ -0,0 +1,24 @@
+// Copyright Epic Games, Inc. All Rights Reserved.
+
+#include "AvaDataLinkProcessor.h"
+#include "AvaSceneSubsystem.h"
+#include "Engine/Level.h"
+#include "Engine/World.h"
+
+IAvaSceneInterface* UAvaDataLinkProcessor::GetSceneInterface() const
+{
+	ULevel* Level = GetTypedOuter<ULevel>();
+	if (!Level || !Level->OwningWorld)
+	{
+		return nullptr;
+	}
+
+	const UAvaSceneSubsystem* SceneSubsystem = Level->OwningWorld->GetSubsystem<UAvaSceneSubsystem>();
+	if (!SceneSubsystem)
+	{
+		// No subsystem yet available, try finding the scene interface by iterating the actors
+		return UAvaSceneSubsystem::FindSceneInterface(Level);
+	}
+
+	return SceneSubsystem->GetSceneInterface(Level);
+}
diff --git a/Engine/Plugins/Experimental/AvalancheDataLink/Source/AvalancheDataLink/Private/AvaDataLinkProcessor.h b/Engine/Plugins/Experimental/AvalancheDataLink/Source/AvalancheDataLink/Private/AvaDataLinkProcessor.h
new file mode 100644
--- /dev/null
+++ b/Engine/Plugins/Experimental/AvalancheDataLink/Source/AvalancheDataLink/Private/AvaDataLinkProcessor.h
@@ -0,0 +1,17 @@
+// Copyright Epic Games, Inc. All Rights Reserved.
+
+#pragma once
+
+#include "DataLinkProcessor.h"
+#include "AvaDataLinkProcessor.generated.h"
+
+class IAvaSceneInterface;
+
+UCLASS(Abstract)
+class UAvaDataLinkProcessor : public UDataLinkProcessor
+{
+	GENERATED_BODY()
+
+protected:
+	IAvaSceneInterface* GetSceneInterface() const;
+};
diff --git a/Engine/Plugins/Experimental/AvalancheDataLink/Source/AvalancheDataLink/Private/RemoteControl/AvaDataLinkRCProcessor.cpp b/Engine/Plugins/Experimental/AvalancheDataLink/Source/AvalancheDataLink/Private/RemoteControl/AvaDataLinkRCProcessor.cpp
new file mode 100644
--- /dev/null
+++ b/Engine/Plugins/Experimental/AvalancheDataLink/Source/AvalancheDataLink/Private/RemoteControl/AvaDataLinkRCProcessor.cpp
@@ -0,0 +1,250 @@
+// Copyright Epic Games, Inc. All Rights Reserved.
+
+#include "AvaDataLinkRCProcessor.h"
+#include "AvaDataLinkLog.h"
+#include "Controller/RCController.h"
+#include "DataLinkJsonUtils.h"
+#include "DataLinkUtils.h"
+#include "Dom/JsonValue.h"
+#include "IAvaSceneInterface.h"
+#include "JsonObjectConverter.h"
+#include "JsonObjectWrapper.h"
+#include "PropertyBindingTypes.h"
+#include "RCVirtualPropertyContainer.h"
+#include "RemoteControlPreset.h"
+
+namespace UE::AvaDataLink::Private
+{
+	struct FCopyInfo
+	{
+		const FProperty* SourceProperty;
+		const uint8* SourceMemory;
+		FProperty* TargetProperty;
+		uint8* TargetMemory;
+	};
+
+	/**
+	 * Tries promoting a given source property value to match the target property type and copies it to the target memory
+	 * Derived from FPropertyBindingBindingCollection::PerformCopy
+	 */
+	template<typename InSourcePropertyType>
+	class FPromotionCopy
+	{
+		template<typename InTargetPropertyType>
+		bool CopySingle(const FCopyInfo& InCopyInfo, const InSourcePropertyType& InSourceProperty)
+		{
+			using FSourceCppType = InSourcePropertyType::TCppType;
+			using FTargetCppType = InTargetPropertyType::TCppType;
+
+			static_assert(std::is_convertible_v<FSourceCppType, FTargetCppType>);
+
+			if (InTargetPropertyType* TargetProperty = CastField<InTargetPropertyType>(InCopyInfo.TargetProperty))
+			{
+				TargetProperty->SetPropertyValue(InCopyInfo.TargetMemory, InSourceProperty.GetPropertyValue(InCopyInfo.SourceMemory));
+				return true;
+			}
+			return false;
+		}
+
+	public:
+		/** Tries promoting from a source property type to any of the given target property types */
+		template<typename ...InTargetPropertyTypes>
+		bool Copy(const FCopyInfo& InCopyInfo)
+		{
+			if (const InSourcePropertyType* SourceProperty = CastField<InSourcePropertyType>(InCopyInfo.SourceProperty))
+			{
+				return (this->CopySingle<InTargetPropertyTypes>(InCopyInfo, *SourceProperty) || ...);
+			}
+			return false;
+		}
+	};
+
+	bool PromoteCopy(const FCopyInfo& InCopyInfo)
+	{
+		// Bool Promotions
+		return FPromotionCopy<FBoolProperty>().Copy<FByteProperty, FIntProperty, FUInt32Property, FInt64Property, FFloatProperty, FDoubleProperty>(InCopyInfo)
+			// Byte Promotion
+			|| FPromotionCopy<FByteProperty>().Copy<FIntProperty, FUInt32Property, FInt64Property, FFloatProperty, FDoubleProperty>(InCopyInfo)
+			// Int32 Promotions
+			|| FPromotionCopy<FIntProperty>().Copy<FInt64Property, FFloatProperty, FDoubleProperty>(InCopyInfo)
+			// Uint32 Promotions
+			|| FPromotionCopy<FUInt32Property>().Copy<FInt64Property, FFloatProperty, FDoubleProperty>(InCopyInfo)
+			// Float Promotions
+			|| FPromotionCopy<FFloatProperty>().Copy<FIntProperty, FInt64Property, FDoubleProperty>(InCopyInfo)
+			// Double Promotions
+			|| FPromotionCopy<FDoubleProperty>().Copy<FIntProperty, FInt64Property, FFloatProperty>(InCopyInfo)
+		;
+	}
+
+} // UE::AvaDataLink::Private
+
+void UAvaDataLinkRCProcessor::OnProcessOutput(FDataLinkProcessorContext& InContext, FConstStructView InOutputDataView)
+{
+	if (!InOutputDataView.IsValid())
+	{
+		return;
+	}
+
+	URemoteControlPreset* const Preset = GetRemoteControlPreset();
+	if (!Preset)
+	{
+		UE_LOG(LogAvaDataLink, Error, TEXT("[%.*s] Data Link execution finished, but Remote Control is invalid!")
+			, InContext.GetContextName().Len()
+			, InContext.GetContextName().GetData());
+		return;
+	}
+
+	URCVirtualPropertyContainerBase* const ControllerContainer = Preset->GetControllerContainer();
+	if (!ControllerContainer)
+	{
+		UE_LOG(LogAvaDataLink, Error, TEXT("[%.*s] Data Link execution finished, but Remote Control '%s' has invalid Controller Container!")
+			, InContext.GetContextName().Len()
+			, InContext.GetContextName().GetData()
+			, *Preset->GetName());
+		return;
+	}
+
+	const FStructView TargetDataView = ControllerContainer->GetPropertyBagMutableValue();
+	if (!TargetDataView.IsValid())
+	{
+		UE_LOG(LogAvaDataLink, Error, TEXT("[%.*s] Data Link execution finished, but Remote Control '%s' has invalid Controller Container Data View!")
+			, InContext.GetContextName().Len()
+			, InContext.GetContextName().GetData()
+			, *Preset->GetName());
+		return;
+	}
+
+	TArray<URCController*> ModifiedControllers;
+	ModifiedControllers.Reserve(ControllerMappings.Num());
+
+	if (const FJsonObjectWrapper* OutputJson = InOutputDataView.GetPtr<const FJsonObjectWrapper>())
+	{
+		if (!OutputJson->JsonObject.IsValid())
+		{
+			UE_LOG(LogAvaDataLink, Warning, TEXT("[%.*s] Data Link output could not be applied to controllers. Json Object was not valid!")
+				, InContext.GetContextName().Len()
+				, InContext.GetContextName().GetData());
+			return;
+		}
+
+		TSharedRef<FJsonObject> JsonObject = OutputJson->JsonObject.ToSharedRef();
+
+		ForEachResolvedController(InContext, Preset, TargetDataView,
+			[&JsonObject, &ModifiedControllers](const FResolvedController& InResolvedController)
+			{
+				TSharedPtr<FJsonValue> SourceJsonValue = UE::DataLinkJson::FindJsonValue(JsonObject, InResolvedController.Mapping->OutputFieldName);
+
+				// If property was set, return true to signal that the controller value was modified
+				if (FJsonObjectConverter::JsonValueToUProperty(SourceJsonValue, InResolvedController.TargetProperty, InResolvedController.TargetMemory))
+				{
+					ModifiedControllers.AddUnique(InResolvedController.Controller);
+				}
+			});
+	}
+	else
+	{
+		ForEachResolvedController(InContext, Preset, TargetDataView,
+			[&InOutputDataView, &ModifiedControllers, &InContext](const FResolvedController& InResolvedController)
+			{
+				FString ErrorMessage;
+				const UE::DataLink::FConstPropertyView SourceView = UE::DataLink::ResolveConstPropertyView(InOutputDataView, InResolvedController.Mapping->OutputFieldName, &ErrorMessage);
+				if (!SourceView.Property)
+				{
+					UE_LOG(LogAvaDataLink, Warning, TEXT("[%.*s] Data Link output field name '%s' could not be applied: %s")
+						, InContext.GetContextName().Len()
+						, InContext.GetContextName().GetData()
+						, *InResolvedController.Mapping->OutputFieldName
+						, *ErrorMessage);
+					return;
+				}
+
+				const UE::PropertyBinding::EPropertyCompatibility Compatibility = UE::PropertyBinding::GetPropertyCompatibility(SourceView.Property, InResolvedController.TargetProperty);
+				if (Compatibility == UE::PropertyBinding::EPropertyCompatibility::Incompatible)
+				{
+					UE_LOG(LogAvaDataLink, Warning, TEXT("[%.*s] Data Link output '%s' could not be applied to controller '%s' as types are incompatible")
+						, InContext.GetContextName().Len()
+						, InContext.GetContextName().GetData()
+						, *InResolvedController.Mapping->OutputFieldName
+						, *InResolvedController.Mapping->TargetController.Name.ToString());
+					return;
+				}
+
+				switch (Compatibility)
+				{
+				case UE::PropertyBinding::EPropertyCompatibility::Compatible:
+					InResolvedController.TargetProperty->CopyCompleteValue(InResolvedController.TargetMemory, SourceView.Memory);
+					ModifiedControllers.AddUnique(InResolvedController.Controller);
+					break;
+
+				case UE::PropertyBinding::EPropertyCompatibility::Promotable:
+					ensureMsgf(UE::AvaDataLink::Private::PromoteCopy({
+						.SourceProperty = SourceView.Property,
+						.SourceMemory = SourceView.Memory,
+						.TargetProperty = InResolvedController.TargetProperty,
+						.TargetMemory = InResolvedController.TargetMemory
+					}), TEXT("Promotion failed even though compatibility was deemed as 'promotable'."));
+					ModifiedControllers.AddUnique(InResolvedController.Controller);
+					break;
+				}
+			});
+	}
+
+	TSet<FGuid> ModifiedControllerIds;
+	ModifiedControllerIds.Reserve(ModifiedControllers.Num());
+
+	for (URCController* Controller : ModifiedControllers)
+	{
+		Controller->OnModifyPropertyValue();
+		ModifiedControllerIds.Add(Controller->Id);
+	}
+
+	Preset->OnControllerModified().Broadcast(Preset, ModifiedControllerIds);
+}
+
+URemoteControlPreset* UAvaDataLinkRCProcessor::GetRemoteControlPreset() const
+{
+	if (IAvaSceneInterface* SceneInterface = GetSceneInterface())
+	{
+		return SceneInterface->GetRemoteControlPreset();
+	}
+	return nullptr;
+}
+
+void UAvaDataLinkRCProcessor::ForEachResolvedController(FDataLinkProcessorContext& InContext, URemoteControlPreset* InPreset, FStructView InTargetDataView, TFunctionRef<void(const FResolvedController&)> InFunction)
+{
+	for (const FAvaDataLinkControllerMapping& Mapping : ControllerMappings)
+	{
+		FResolvedController ResolvedMapping;
+		ResolvedMapping.Mapping = &Mapping;
+
+		// Get the controller to retrieve the underlying Property FName of the controller within the Controller Container Property Bag
+		ResolvedMapping.Controller = Cast<URCController>(Mapping.TargetController.FindController(InPreset));
+		if (!ResolvedMapping.Controller)
+		{
+			UE_LOG(LogAvaDataLink, Warning, TEXT("[%.*s] Data Link output '%s' could not be applied to controller '%s'. Controller was not found in preset '%s'!")
+				, InContext.GetContextName().Len()
+				, InContext.GetContextName().GetData()
+				, *Mapping.OutputFieldName
+				, *Mapping.TargetController.Name.ToString()
+				, *InPreset->GetName());
+			continue;
+		}
+
+		const UE::DataLink::FPropertyView TargetPropertyView = UE::DataLink::ResolvePropertyView(InTargetDataView, ResolvedMapping.Controller->PropertyName.ToString());
+		if (!TargetPropertyView.Property)
+		{
+			UE_LOG(LogAvaDataLink, Warning, TEXT("[%.*s] Data Link output '%s' could not be applied to controller '%s'. Controller property '%s' in preset '%s' was not found!")
+				, InContext.GetContextName().Len()
+				, InContext.GetContextName().GetData()
+				, *Mapping.OutputFieldName
+				, *Mapping.TargetController.Name.ToString()
+				, *ResolvedMapping.Controller->PropertyName.ToString()
+				, *InPreset->GetName());
+			continue;
+		}
+
+		ResolvedMapping.TargetProperty = TargetPropertyView.Property;
+		ResolvedMapping.TargetMemory = TargetPropertyView.Memory;
+		InFunction(ResolvedMapping);
+	}
+}
diff --git a/Engine/Plugins/Experimental/AvalancheDataLink/Source/AvalancheDataLink/Private/RemoteControl/AvaDataLinkRCProcessor.h b/Engine/Plugins/Experimental/AvalancheDataLink/Source/AvalancheDataLink/Private/RemoteControl/AvaDataLinkRCProcessor.h
new file mode 100644
--- /dev/null
+++ b/Engine/Plugins/Experimental/AvalancheDataLink/Source/AvalancheDataLink/Private/RemoteControl/AvaDataLinkRCProcessor.h
@@ -0,0 +1,42 @@
+// Copyright Epic Games, Inc. All Rights Reserved.
+
+#pragma once
+
+#include "AvaDataLinkProcessor.h"
+#include "RemoteControl/AvaDataLinkControllerMapping.h"
+#include "AvaDataLinkRCProcessor.generated.h"
+
+class FProperty;
+class URCController;
+class URemoteControlPreset;
+struct FAvaDataLinkControllerMapping;
+
+UCLASS(DisplayName="Motion Design Data Link Remote Control Processor")
+class UAvaDataLinkRCProcessor : public UAvaDataLinkProcessor
+{
+	GENERATED_BODY()
+
+protected:
+	//~ Begin UDataLinkProcessor
+	virtual void OnProcessOutput(FDataLinkProcessorContext& InContext, FConstStructView InOutputDataView) override;
+	//~ End UDataLinkProcessor
+
+	URemoteControlPreset* GetRemoteControlPreset() const;
+
+	struct FResolvedController
+	{
+		const FAvaDataLinkControllerMapping* Mapping;
+		TObjectPtr<URCController> Controller;
+		FProperty* TargetProperty;
+		uint8* TargetMemory;
+	};
+
+	/** Iterates each controller resolved from the controller mappings */
+	void ForEachResolvedController(FDataLinkProcessorContext& InContext, URemoteControlPreset* InPreset, FStructView InTargetDataView, TFunctionRef<void(const FResolvedController&)> InFunction);
+
+private:
+	UPROPERTY(EditAnywhere, BlueprintReadWrite, Category="Motion Design Data Link", DisplayName="Output Field to Controller Mappings", meta=(AllowPrivateAccess="true"))
+	TArray<FAvaDataLinkControllerMapping> ControllerMappings;
+
+	friend class UAvaDataLinkInstance;
+};
diff --git a/Engine/Plugins/Experimental/AvalancheDataLink/Source/AvalancheDataLink/Public/AvaDataLinkInstance.h b/Engine/Plugins/Experimental/AvalancheDataLink/Source/AvalancheDataLink/Public/AvaDataLinkInstance.h
--- a/Engine/Plugins/Experimental/AvalancheDataLink/Source/AvalancheDataLink/Public/AvaDataLinkInstance.h
+++ b/Engine/Plugins/Experimental/AvalancheDataLink/Source/AvalancheDataLink/Public/AvaDataLinkInstance.h
@@ -2,78 +2,63 @@
 
 #pragma once
 
-#include "AvaRCControllerId.h"
 #include "DataLinkInstance.h"
+#include "RemoteControl/AvaDataLinkControllerMapping.h"
 #include "UObject/Object.h"
 #include "AvaDataLinkInstance.generated.h"
 
-class FDataLinkExecutor;
-class IAvaSceneInterface;
-class URCController;
-enum class EDataLinkExecutionResult : uint8;
-struct FConstStructView;
-struct FStructView;
-
-USTRUCT(BlueprintType, DisplayName="Motion Design Data Link RC Controller Mapping")
-struct FAvaDataLinkControllerMapping
-{
-	GENERATED_BODY()
-
-	UPROPERTY(EditAnywhere, BlueprintReadWrite, Category = "Motion Design Data Link")
-	FString OutputFieldName;
+#define UE_API AVALANCHEDATALINK_API
 
-	UPROPERTY(EditAnywhere, BlueprintReadWrite, Category = "Motion Design Data Link")
-	FAvaRCControllerId TargetController;
-};
+class FDataLinkExecutor;
+class UDataLinkProcessor;
 
-UCLASS(MinimalAPI, EditInlineNew)
+UCLASS(MinimalAPI, EditInlineNew, DisplayName="Motion Design Data Link Instance")
 class UAvaDataLinkInstance : public UObject
 {
 	GENERATED_BODY()
 
 public:
  	/** Starts a new active data link execution, stopping the existing if any */
-	void Execute();
+	UE_API void Execute();
 
 	/** Stops the active data link execution if any */
-	void Stop();
+	UE_API void Stop();
 
 	static FName GetDataLinkInstancePropertyName()
 	{
 		return GET_MEMBER_NAME_CHECKED(UAvaDataLinkInstance, DataLinkInstance);
 	}
 
-	static FName GetControllerMappingsPropertyName()
+	static FName GetOutputProcessorsPropertyName()
 	{
-		return GET_MEMBER_NAME_CHECKED(UAvaDataLinkInstance, ControllerMappings);
+		return GET_MEMBER_NAME_CHECKED(UAvaDataLinkInstance, OutputProcessors);
 	}
 
+	//~ Begin UObject
+#if WITH_EDITORONLY_DATA
+	static UE_API void DeclareConstructClasses(TArray<FTopLevelAssetPath>& OutConstructClasses, const UClass* InSpecificSubclass);
+#endif
+	UE_API virtual void PostLoad() override;
+	//~ End UObject
+
 private:
 #if WITH_DATALINK_CONTEXT
 	FString BuildContextName() const;
 #endif
 
-	/** Called when data link execution has received data */
-	void OnOutputData(const FDataLinkExecutor& InExecutor, FConstStructView InOutputDataView);
-
-	IAvaSceneInterface* GetSceneInterface() const;
-
-	URemoteControlPreset* GetRemoteControlPreset() const;
-
-	struct FResolvedController
-	{
-		const FAvaDataLinkControllerMapping* Mapping;
-		TObjectPtr<URCController> Controller;
-		FProperty* TargetProperty;
-		uint8* TargetMemory;
-	};
-	void ForEachResolvedController(const FDataLinkExecutor& InExecutor, URemoteControlPreset* InPreset, FStructView InTargetDataView, TFunctionRef<void(const FResolvedController&)> InFunction);
-
 	UPROPERTY(EditAnywhere, BlueprintReadWrite, Category="Motion Design Data Link", meta=(AllowPrivateAccess="true"))
 	FDataLinkInstance DataLinkInstance;
 
-	UPROPERTY(EditAnywhere, BlueprintReadWrite, Category="Motion Design Data Link", DisplayName="Output Field to Controller Mappings", meta=(AllowPrivateAccess="true"))
-	TArray<FAvaDataLinkControllerMapping> ControllerMappings;
+	UPROPERTY(EditAnywhere, BlueprintReadWrite, Instanced, Category="Motion Design Data Link", meta=(AllowPrivateAccess="true"))
+	TArray<TObjectPtr<UDataLinkProcessor>> OutputProcessors;
+
+#if WITH_EDITORONLY_DATA
+	UE_DEPRECATED(5.7, "ControllerMappings in Data Link Instance is deprecated. Use Output Processors instead")
+	UPROPERTY()
+	TArray<FAvaDataLinkControllerMapping> ControllerMappings_DEPRECATED;
+#endif
 
 	TSharedPtr<FDataLinkExecutor> Executor;
 };
+
+#undef UE_API
diff --git a/Engine/Plugins/Experimental/AvalancheDataLink/Source/AvalancheDataLink/Public/RemoteControl/AvaDataLinkControllerMapping.h b/Engine/Plugins/Experimental/AvalancheDataLink/Source/AvalancheDataLink/Public/RemoteControl/AvaDataLinkControllerMapping.h
new file mode 100644
--- /dev/null
+++ b/Engine/Plugins/Experimental/AvalancheDataLink/Source/AvalancheDataLink/Public/RemoteControl/AvaDataLinkControllerMapping.h
@@ -0,0 +1,18 @@
+// Copyright Epic Games, Inc. All Rights Reserved.
+
+#pragma once
+
+#include "AvaRCControllerId.h"
+#include "AvaDataLinkControllerMapping.generated.h"
+
+USTRUCT(BlueprintType, DisplayName="Motion Design Data Link RC Controller Mapping")
+struct FAvaDataLinkControllerMapping
+{
+	GENERATED_BODY()
+
+	UPROPERTY(EditAnywhere, BlueprintReadWrite, Category = "Motion Design Data Link")
+	FString OutputFieldName;
+
+	UPROPERTY(EditAnywhere, BlueprintReadWrite, Category = "Motion Design Data Link")
+	FAvaRCControllerId TargetController;
+};
diff --git a/Engine/Plugins/Experimental/AvalancheDataLink/Source/AvalancheDataLinkEditor/Private/DetailsView/AvaDataLinkInstanceCustomization.cpp b/Engine/Plugins/Experimental/AvalancheDataLink/Source/AvalancheDataLinkEditor/Private/DetailsView/AvaDataLinkInstanceCustomization.cpp
--- a/Engine/Plugins/Experimental/AvalancheDataLink/Source/AvalancheDataLinkEditor/Private/DetailsView/AvaDataLinkInstanceCustomization.cpp
+++ b/Engine/Plugins/Experimental/AvalancheDataLink/Source/AvalancheDataLinkEditor/Private/DetailsView/AvaDataLinkInstanceCustomization.cpp
@@ -1,8 +1,6 @@
 // Copyright Epic Games, Inc. All Rights Reserved.
 
 #include "AvaDataLinkInstanceCustomization.h"
-
-#include "AvaDataLinkControllerMappingsBuilder.h"
 #include "AvaDataLinkInstance.h"
 #include "DetailCategoryBuilder.h"
 #include "DetailLayoutBuilder.h"
@@ -14,10 +12,10 @@ void FAvaDataLinkInstanceCustomization::CustomizeDetails(IDetailLayoutBuilder& I
 	TSharedRef<IPropertyHandle> DataLinkInstanceHandle = InDetailBuilder.GetProperty(UAvaDataLinkInstance::GetDataLinkInstancePropertyName());
 	DataLinkInstanceHandle->MarkHiddenByCustomization();
 
-	TSharedRef<IPropertyHandle> ControllerMappingsHandle = InDetailBuilder.GetProperty(UAvaDataLinkInstance::GetControllerMappingsPropertyName());
-	ControllerMappingsHandle->MarkHiddenByCustomization();
+	TSharedRef<IPropertyHandle> OutputProcessorsHandle = InDetailBuilder.GetProperty(UAvaDataLinkInstance::GetOutputProcessorsPropertyName());
+	OutputProcessorsHandle->MarkHiddenByCustomization();
 
 	IDetailCategoryBuilder& Category = InDetailBuilder.EditCategory(DataLinkInstanceHandle->GetDefaultCategoryName(), DataLinkInstanceHandle->GetDefaultCategoryText());
 	Category.AddProperty(DataLinkInstanceHandle);
-	Category.AddCustomBuilder(MakeShared<FAvaDataLinkControllerMappingsBuilder>(ControllerMappingsHandle, InDetailBuilder.GetDetailsViewSharedPtr()));
+	Category.AddProperty(OutputProcessorsHandle);
 }
diff --git a/Engine/Plugins/Experimental/DataLink/Source/DataLink/Private/DataLinkExecutor.cpp b/Engine/Plugins/Experimental/DataLink/Source/DataLink/Private/DataLinkExecutor.cpp
--- a/Engine/Plugins/Experimental/DataLink/Source/DataLink/Private/DataLinkExecutor.cpp
+++ b/Engine/Plugins/Experimental/DataLink/Source/DataLink/Private/DataLinkExecutor.cpp
@@ -7,6 +7,7 @@
 #include "DataLinkLog.h"
 #include "DataLinkNode.h"
 #include "DataLinkPinReference.h"
+#include "DataLinkProcessor.h"
 #include "DataLinkSink.h"
 #include "HAL/IConsoleManager.h"
 #include "Templates/SubclassOf.h"
@@ -29,9 +30,17 @@ TSharedPtr<FDataLinkExecutor> FDataLinkExecutor::Create(FDataLinkExecutorArgumen
 	Executor->ContextObject = MoveTemp(InArgs.ContextObject);
 	Executor->Instance = MoveTemp(InArgs.Instance);
 	Executor->Sink = MoveTemp(InArgs.Sink);
+	Executor->OutputProcessors = MoveTemp(InArgs.OutputProcessors);
 	Executor->OnOutputData = MoveTemp(InArgs.OnOutputData);
 	Executor->OnExecutionFinished = MoveTemp(InArgs.OnExecutionFinished);
 
+	// Remove invalid entries
+	Executor->OutputProcessors.RemoveAll(
+		[](UDataLinkProcessor* InProcessor)
+		{
+			return !InProcessor;
+		});
+
 	return Executor;
 }
 
@@ -87,6 +96,8 @@ void FDataLinkExecutor::AddReferencedObjects(FReferenceCollector& InCollector)
 	{
 		Pair.Value.AddReferencedObjects(InCollector);
 	}
+
+	InCollector.AddReferencedObjects(OutputProcessors);
 }
 
 void FDataLinkExecutor::Run()
@@ -487,6 +498,14 @@ void FDataLinkExecutor::BroadcastOutputData(FConstStructView InOutputData)
 		return;
 	}
 
+	{
+		FDataLinkProcessorContext ProcessorContext(*this);
+		for (UDataLinkProcessor* OutputProcessor : OutputProcessors)
+		{
+			OutputProcessor->ProcessOutput(ProcessorContext, InOutputData);
+		}
+	}
+
 	OnOutputData.ExecuteIfBound(*this, InOutputData);
 
 	if (!bPersistExecution)
diff --git a/Engine/Plugins/Experimental/DataLink/Source/DataLink/Private/DataLinkExecutorArguments.cpp b/Engine/Plugins/Experimental/DataLink/Source/DataLink/Private/DataLinkExecutorArguments.cpp
--- a/Engine/Plugins/Experimental/DataLink/Source/DataLink/Private/DataLinkExecutorArguments.cpp
+++ b/Engine/Plugins/Experimental/DataLink/Source/DataLink/Private/DataLinkExecutorArguments.cpp
@@ -39,6 +39,12 @@ FDataLinkExecutorArguments&& FDataLinkExecutorArguments::SetSink(const TSharedPt
 	return MoveTemp(*this);
 }
 
+FDataLinkExecutorArguments&& FDataLinkExecutorArguments::SetOutputProcessors(TConstArrayView<UDataLinkProcessor*> InOutputProcessors)
+{
+	OutputProcessors = InOutputProcessors;
+	return MoveTemp(*this);
+}
+
 FDataLinkExecutorArguments&& FDataLinkExecutorArguments::SetOnOutputData(FOnDataLinkOutputData&& InDelegate)
 {
 	OnOutputData = MoveTemp(InDelegate);
diff --git a/Engine/Plugins/Experimental/DataLink/Source/DataLink/Private/DataLinkProcessor.cpp b/Engine/Plugins/Experimental/DataLink/Source/DataLink/Private/DataLinkProcessor.cpp
new file mode 100644
--- /dev/null
+++ b/Engine/Plugins/Experimental/DataLink/Source/DataLink/Private/DataLinkProcessor.cpp
@@ -0,0 +1,24 @@
+// Copyright Epic Games, Inc. All Rights Reserved.
+
+#include "DataLinkProcessor.h"
+#include "DataLinkExecutor.h"
+
+FDataLinkProcessorContext::FDataLinkProcessorContext(const FDataLinkExecutor& InExecutor)
+	: Executor(InExecutor)
+{
+}
+
+FStringView FDataLinkProcessorContext::GetContextName() const
+{
+	return Executor.GetContextName();
+}
+
+UObject* FDataLinkProcessorContext::GetContextObject() const
+{
+	return Executor.GetContextObject();
+}
+
+void UDataLinkProcessor::ProcessOutput(FDataLinkProcessorContext& InContext, FConstStructView InOutputDataView)
+{
+	OnProcessOutput(InContext, InOutputDataView);
+}
diff --git a/Engine/Plugins/Experimental/DataLink/Source/DataLink/Public/DataLinkExecutor.h b/Engine/Plugins/Experimental/DataLink/Source/DataLink/Public/DataLinkExecutor.h
--- a/Engine/Plugins/Experimental/DataLink/Source/DataLink/Public/DataLinkExecutor.h
+++ b/Engine/Plugins/Experimental/DataLink/Source/DataLink/Public/DataLinkExecutor.h
@@ -12,6 +12,7 @@
 class FDataLinkExecutorArguments;
 class UDataLinkGraph;
 class UDataLinkNode;
+class UDataLinkProcessor;
 struct FDataLinkPinReference;
 
 /**
@@ -146,6 +147,9 @@ private:
 	/** Sink where all the data for this execution is queried and stored */
 	TSharedPtr<FDataLinkSink> Sink;
 
+	/** Container for all the output processors in this execution */
+	TArray<TObjectPtr<UDataLinkProcessor>> OutputProcessors;
+
 	/** Delegate to call when the data link outputs data */
 	FOnDataLinkOutputData OnOutputData;
 
diff --git a/Engine/Plugins/Experimental/DataLink/Source/DataLink/Public/DataLinkExecutorArguments.h b/Engine/Plugins/Experimental/DataLink/Source/DataLink/Public/DataLinkExecutorArguments.h
--- a/Engine/Plugins/Experimental/DataLink/Source/DataLink/Public/DataLinkExecutorArguments.h
+++ b/Engine/Plugins/Experimental/DataLink/Source/DataLink/Public/DataLinkExecutorArguments.h
@@ -8,6 +8,7 @@
 #include "StructUtils/StructView.h"
 
 class UDataLinkGraph;
+class UDataLinkProcessor;
 struct FDataLinkSink;
 
 class FDataLinkExecutorArguments
@@ -32,6 +33,9 @@ public:
 	/** Sets the Sink to use. If unset, the Executor will create its own temporary sink */
 	DATALINK_API FDataLinkExecutorArguments&& SetSink(const TSharedPtr<FDataLinkSink>& InSink);
 
+	/** Sets the output processors */
+	DATALINK_API FDataLinkExecutorArguments&& SetOutputProcessors(TConstArrayView<UDataLinkProcessor*> InOutputProcessors);
+
 	/** The delegate to call when the execution outputs data */
 	DATALINK_API FDataLinkExecutorArguments&& SetOnOutputData(FOnDataLinkOutputData&& InDelegate);
 
@@ -49,6 +53,9 @@ private:
 
 	TSharedPtr<FDataLinkSink> Sink;
 
+	/** The output processors to use for the execution */
+	TArray<TObjectPtr<UDataLinkProcessor>> OutputProcessors;
+
 	/** Delegate called when data link execution has output data ready */
 	FOnDataLinkOutputData OnOutputData;
 
diff --git a/Engine/Plugins/Experimental/DataLink/Source/DataLink/Public/DataLinkProcessor.h b/Engine/Plugins/Experimental/DataLink/Source/DataLink/Public/DataLinkProcessor.h
new file mode 100644
--- /dev/null
+++ b/Engine/Plugins/Experimental/DataLink/Source/DataLink/Public/DataLinkProcessor.h
@@ -0,0 +1,45 @@
+// Copyright Epic Games, Inc. All Rights Reserved.
+
+#pragma once
+
+#include "StructUtils/StructView.h"
+#include "UObject/Object.h"
+#include "DataLinkProcessor.generated.h"
+
+class FDataLinkExecutor;
+
+/** Short-lived object wrapping its owning executor for processors to call */
+struct FDataLinkProcessorContext
+{
+	explicit FDataLinkProcessorContext(const FDataLinkExecutor& InExecutor);
+
+	/**
+	 * Gets a brief description of the context for debugging purposes.
+	 * Returns a valid context name when WITH_DATALINK_CONTEXT is true.
+	 * Note: WITH_DATALINK_CONTEXT is true in builds with logging enabled.
+	 */
+	DATALINK_API FStringView GetContextName() const;
+
+	/** Gets the object that is responsible for this execution */
+	DATALINK_API UObject* GetContextObject() const;
+
+private:
+	/** Executor that owns this context */
+	const FDataLinkExecutor& Executor;
+};
+
+/** Processors execute logic on the output data of a data link graph execution */
+UCLASS(MinimalAPI, Abstract, EditInlineNew)
+class UDataLinkProcessor : public UObject
+{
+	GENERATED_BODY()
+
+public:
+	/** Called when the data link graph execution has output data */
+	void ProcessOutput(FDataLinkProcessorContext& InContext, FConstStructView InOutputDataView);
+
+protected:
+	virtual void OnProcessOutput(FDataLinkProcessorContext& InContext, FConstStructView InOutputDataView)
+	{
+	}
+};
