Motion Design Data Link: added blueprintable support for processors, added initialize/finalize calls for data link nodes

diff --git a/Engine/Plugins/Experimental/AvalancheDataLink/Source/AvalancheDataLink/Private/RemoteControl/AvaDataLinkRCProcessor.cpp b/Engine/Plugins/Experimental/AvalancheDataLink/Source/AvalancheDataLink/Private/RemoteControl/AvaDataLinkRCProcessor.cpp
--- a/Engine/Plugins/Experimental/AvalancheDataLink/Source/AvalancheDataLink/Private/RemoteControl/AvaDataLinkRCProcessor.cpp
+++ b/Engine/Plugins/Experimental/AvalancheDataLink/Source/AvalancheDataLink/Private/RemoteControl/AvaDataLinkRCProcessor.cpp
@@ -3,6 +3,7 @@
 #include "AvaDataLinkRCProcessor.h"
 #include "AvaDataLinkLog.h"
 #include "Controller/RCController.h"
+#include "DataLinkExecutor.h"
 #include "DataLinkJsonUtils.h"
 #include "DataLinkUtils.h"
 #include "Dom/JsonValue.h"
@@ -78,7 +79,7 @@ namespace UE::AvaDataLink::Private
 
 } // UE::AvaDataLink::Private
 
-void UAvaDataLinkRCProcessor::OnProcessOutput(FDataLinkProcessorContext& InContext, FConstStructView InOutputDataView)
+void UAvaDataLinkRCProcessor::OnProcessOutput(const FDataLinkExecutor& InExecutor, FConstStructView InOutputDataView)
 {
 	if (!InOutputDataView.IsValid())
 	{
@@ -89,8 +90,8 @@ void UAvaDataLinkRCProcessor::OnProcessOutput(FDataLinkProcessorContext& InConte
 	if (!Preset)
 	{
 		UE_LOG(LogAvaDataLink, Error, TEXT("[%.*s] Data Link execution finished, but Remote Control is invalid!")
-			, InContext.GetContextName().Len()
-			, InContext.GetContextName().GetData());
+			, InExecutor.GetContextName().Len()
+			, InExecutor.GetContextName().GetData());
 		return;
 	}
 
@@ -98,8 +99,8 @@ void UAvaDataLinkRCProcessor::OnProcessOutput(FDataLinkProcessorContext& InConte
 	if (!ControllerContainer)
 	{
 		UE_LOG(LogAvaDataLink, Error, TEXT("[%.*s] Data Link execution finished, but Remote Control '%s' has invalid Controller Container!")
-			, InContext.GetContextName().Len()
-			, InContext.GetContextName().GetData()
+			, InExecutor.GetContextName().Len()
+			, InExecutor.GetContextName().GetData()
 			, *Preset->GetName());
 		return;
 	}
@@ -108,8 +109,8 @@ void UAvaDataLinkRCProcessor::OnProcessOutput(FDataLinkProcessorContext& InConte
 	if (!TargetDataView.IsValid())
 	{
 		UE_LOG(LogAvaDataLink, Error, TEXT("[%.*s] Data Link execution finished, but Remote Control '%s' has invalid Controller Container Data View!")
-			, InContext.GetContextName().Len()
-			, InContext.GetContextName().GetData()
+			, InExecutor.GetContextName().Len()
+			, InExecutor.GetContextName().GetData()
 			, *Preset->GetName());
 		return;
 	}
@@ -122,14 +123,14 @@ void UAvaDataLinkRCProcessor::OnProcessOutput(FDataLinkProcessorContext& InConte
 		if (!OutputJson->JsonObject.IsValid())
 		{
 			UE_LOG(LogAvaDataLink, Warning, TEXT("[%.*s] Data Link output could not be applied to controllers. Json Object was not valid!")
-				, InContext.GetContextName().Len()
-				, InContext.GetContextName().GetData());
+				, InExecutor.GetContextName().Len()
+				, InExecutor.GetContextName().GetData());
 			return;
 		}
 
 		TSharedRef<FJsonObject> JsonObject = OutputJson->JsonObject.ToSharedRef();
 
-		ForEachResolvedController(InContext, Preset, TargetDataView,
+		ForEachResolvedController(InExecutor, Preset, TargetDataView,
 			[&JsonObject, &ModifiedControllers](const FResolvedController& InResolvedController)
 			{
 				TSharedPtr<FJsonValue> SourceJsonValue = UE::DataLinkJson::FindJsonValue(JsonObject, InResolvedController.Mapping->OutputFieldName);
@@ -143,16 +144,16 @@ void UAvaDataLinkRCProcessor::OnProcessOutput(FDataLinkProcessorContext& InConte
 	}
 	else
 	{
-		ForEachResolvedController(InContext, Preset, TargetDataView,
-			[&InOutputDataView, &ModifiedControllers, &InContext](const FResolvedController& InResolvedController)
+		ForEachResolvedController(InExecutor, Preset, TargetDataView,
+			[&InOutputDataView, &ModifiedControllers, &InExecutor](const FResolvedController& InResolvedController)
 			{
 				FString ErrorMessage;
 				const UE::DataLink::FConstPropertyView SourceView = UE::DataLink::ResolveConstPropertyView(InOutputDataView, InResolvedController.Mapping->OutputFieldName, &ErrorMessage);
 				if (!SourceView.Property)
 				{
 					UE_LOG(LogAvaDataLink, Warning, TEXT("[%.*s] Data Link output field name '%s' could not be applied: %s")
-						, InContext.GetContextName().Len()
-						, InContext.GetContextName().GetData()
+						, InExecutor.GetContextName().Len()
+						, InExecutor.GetContextName().GetData()
 						, *InResolvedController.Mapping->OutputFieldName
 						, *ErrorMessage);
 					return;
@@ -162,8 +163,8 @@ void UAvaDataLinkRCProcessor::OnProcessOutput(FDataLinkProcessorContext& InConte
 				if (Compatibility == UE::PropertyBinding::EPropertyCompatibility::Incompatible)
 				{
 					UE_LOG(LogAvaDataLink, Warning, TEXT("[%.*s] Data Link output '%s' could not be applied to controller '%s' as types are incompatible")
-						, InContext.GetContextName().Len()
-						, InContext.GetContextName().GetData()
+						, InExecutor.GetContextName().Len()
+						, InExecutor.GetContextName().GetData()
 						, *InResolvedController.Mapping->OutputFieldName
 						, *InResolvedController.Mapping->TargetController.Name.ToString());
 					return;
@@ -210,7 +211,7 @@ URemoteControlPreset* UAvaDataLinkRCProcessor::GetRemoteControlPreset() const
 	return nullptr;
 }
 
-void UAvaDataLinkRCProcessor::ForEachResolvedController(FDataLinkProcessorContext& InContext, URemoteControlPreset* InPreset, FStructView InTargetDataView, TFunctionRef<void(const FResolvedController&)> InFunction)
+void UAvaDataLinkRCProcessor::ForEachResolvedController(const FDataLinkExecutor& InExecutor, URemoteControlPreset* InPreset, FStructView InTargetDataView, TFunctionRef<void(const FResolvedController&)> InFunction)
 {
 	for (const FAvaDataLinkControllerMapping& Mapping : ControllerMappings)
 	{
@@ -222,8 +223,8 @@ void UAvaDataLinkRCProcessor::ForEachResolvedController(FDataLinkProcessorContex
 		if (!ResolvedMapping.Controller)
 		{
 			UE_LOG(LogAvaDataLink, Warning, TEXT("[%.*s] Data Link output '%s' could not be applied to controller '%s'. Controller was not found in preset '%s'!")
-				, InContext.GetContextName().Len()
-				, InContext.GetContextName().GetData()
+				, InExecutor.GetContextName().Len()
+				, InExecutor.GetContextName().GetData()
 				, *Mapping.OutputFieldName
 				, *Mapping.TargetController.Name.ToString()
 				, *InPreset->GetName());
@@ -234,8 +235,8 @@ void UAvaDataLinkRCProcessor::ForEachResolvedController(FDataLinkProcessorContex
 		if (!TargetPropertyView.Property)
 		{
 			UE_LOG(LogAvaDataLink, Warning, TEXT("[%.*s] Data Link output '%s' could not be applied to controller '%s'. Controller property '%s' in preset '%s' was not found!")
-				, InContext.GetContextName().Len()
-				, InContext.GetContextName().GetData()
+				, InExecutor.GetContextName().Len()
+				, InExecutor.GetContextName().GetData()
 				, *Mapping.OutputFieldName
 				, *Mapping.TargetController.Name.ToString()
 				, *ResolvedMapping.Controller->PropertyName.ToString()
diff --git a/Engine/Plugins/Experimental/AvalancheDataLink/Source/AvalancheDataLink/Private/RemoteControl/AvaDataLinkRCProcessor.h b/Engine/Plugins/Experimental/AvalancheDataLink/Source/AvalancheDataLink/Private/RemoteControl/AvaDataLinkRCProcessor.h
--- a/Engine/Plugins/Experimental/AvalancheDataLink/Source/AvalancheDataLink/Private/RemoteControl/AvaDataLinkRCProcessor.h
+++ b/Engine/Plugins/Experimental/AvalancheDataLink/Source/AvalancheDataLink/Private/RemoteControl/AvaDataLinkRCProcessor.h
@@ -18,7 +18,7 @@ class UAvaDataLinkRCProcessor : public UAvaDataLinkProcessor
 
 protected:
 	//~ Begin UDataLinkProcessor
-	virtual void OnProcessOutput(FDataLinkProcessorContext& InContext, FConstStructView InOutputDataView) override;
+	virtual void OnProcessOutput(const FDataLinkExecutor& InExecutor, FConstStructView InOutputDataView) override;
 	//~ End UDataLinkProcessor
 
 	URemoteControlPreset* GetRemoteControlPreset() const;
@@ -32,7 +32,7 @@ protected:
 	};
 
 	/** Iterates each controller resolved from the controller mappings */
-	void ForEachResolvedController(FDataLinkProcessorContext& InContext, URemoteControlPreset* InPreset, FStructView InTargetDataView, TFunctionRef<void(const FResolvedController&)> InFunction);
+	void ForEachResolvedController(const FDataLinkExecutor& InExecutor, URemoteControlPreset* InPreset, FStructView InTargetDataView, TFunctionRef<void(const FResolvedController&)> InFunction);
 
 private:
 	UPROPERTY(EditAnywhere, BlueprintReadWrite, Category="Motion Design Data Link", DisplayName="Output Field to Controller Mappings", meta=(AllowPrivateAccess="true"))
diff --git a/Engine/Plugins/Experimental/DataLink/Source/DataLink/Private/DataLinkExecutor.cpp b/Engine/Plugins/Experimental/DataLink/Source/DataLink/Private/DataLinkExecutor.cpp
--- a/Engine/Plugins/Experimental/DataLink/Source/DataLink/Private/DataLinkExecutor.cpp
+++ b/Engine/Plugins/Experimental/DataLink/Source/DataLink/Private/DataLinkExecutor.cpp
@@ -123,6 +123,11 @@ void FDataLinkExecutor::Run()
 	// Hold onto an extra shared reference in case the execution finishes immediately
 	TSharedRef<FDataLinkExecutor> This = SharedThis(this);
 
+	for (UDataLinkProcessor* OutputProcessor : OutputProcessors)
+	{
+		OutputProcessor->Initialize(*this);
+	}
+
 	if (!ExecuteEntryNodes())
 	{
 		BroadcastFailure();
@@ -142,6 +147,12 @@ void FDataLinkExecutor::Stop(EDataLinkExecutionResult InExecutionResult)
 	// Hold onto an extra shared reference to allow implementers to reset their held shared reference in the finish delegate
 	// without destroying this
 	const TSharedRef<FDataLinkExecutor> This = SharedThis(this);
+
+	for (UDataLinkProcessor* OutputProcessor : OutputProcessors)
+	{
+		OutputProcessor->Finalize(*this, InExecutionResult);
+	}
+
 	OnExecutionFinished.ExecuteIfBound(*this, InExecutionResult);
 }
 
@@ -498,12 +509,9 @@ void FDataLinkExecutor::BroadcastOutputData(FConstStructView InOutputData)
 		return;
 	}
 
+	for (UDataLinkProcessor* OutputProcessor : OutputProcessors)
 	{
-		FDataLinkProcessorContext ProcessorContext(*this);
-		for (UDataLinkProcessor* OutputProcessor : OutputProcessors)
-		{
-			OutputProcessor->ProcessOutput(ProcessorContext, InOutputData);
-		}
+		OutputProcessor->ProcessOutput(*this, InOutputData);
 	}
 
 	OnOutputData.ExecuteIfBound(*this, InOutputData);
diff --git a/Engine/Plugins/Experimental/DataLink/Source/DataLink/Private/DataLinkProcessor.cpp b/Engine/Plugins/Experimental/DataLink/Source/DataLink/Private/DataLinkProcessor.cpp
--- a/Engine/Plugins/Experimental/DataLink/Source/DataLink/Private/DataLinkProcessor.cpp
+++ b/Engine/Plugins/Experimental/DataLink/Source/DataLink/Private/DataLinkProcessor.cpp
@@ -3,22 +3,17 @@
 #include "DataLinkProcessor.h"
 #include "DataLinkExecutor.h"
 
-FDataLinkProcessorContext::FDataLinkProcessorContext(const FDataLinkExecutor& InExecutor)
-	: Executor(InExecutor)
+void UDataLinkProcessor::Initialize(const FDataLinkExecutor& InExecutor)
 {
+	OnInitialize(InExecutor);
 }
 
-FStringView FDataLinkProcessorContext::GetContextName() const
+void UDataLinkProcessor::ProcessOutput(const FDataLinkExecutor& InExecutor, FConstStructView InOutputDataView)
 {
-	return Executor.GetContextName();
+	OnProcessOutput(InExecutor, InOutputDataView);
 }
 
-UObject* FDataLinkProcessorContext::GetContextObject() const
+void UDataLinkProcessor::Finalize(const FDataLinkExecutor& InExecutor, EDataLinkExecutionResult InExecutionResult)
 {
-	return Executor.GetContextObject();
-}
-
-void UDataLinkProcessor::ProcessOutput(FDataLinkProcessorContext& InContext, FConstStructView InOutputDataView)
-{
-	OnProcessOutput(InContext, InOutputDataView);
+	OnFinalize(InExecutor, InExecutionResult);
 }
diff --git a/Engine/Plugins/Experimental/DataLink/Source/DataLink/Private/DataLinkProcessorBlueprint.cpp b/Engine/Plugins/Experimental/DataLink/Source/DataLink/Private/DataLinkProcessorBlueprint.cpp
new file mode 100644
--- /dev/null
+++ b/Engine/Plugins/Experimental/DataLink/Source/DataLink/Private/DataLinkProcessorBlueprint.cpp
@@ -0,0 +1,41 @@
+// Copyright Epic Games, Inc. All Rights Reserved.
+
+#include "DataLinkProcessorBlueprint.h"
+#include "DataLinkExecutor.h"
+
+UObject* UDataLinkProcessorBlueprint::GetContextObject() const
+{
+	if (TSharedPtr<const FDataLinkExecutor> Executor = ExecutorWeak.Pin())
+	{
+		return Executor->GetContextObject();
+	}
+	return nullptr;
+}
+
+void UDataLinkProcessorBlueprint::OnInitialize(const FDataLinkExecutor& InExecutor)
+{
+	ExecutorWeak = InExecutor.AsWeak();
+	Initialize();
+}
+
+void UDataLinkProcessorBlueprint::OnProcessOutput(const FDataLinkExecutor& InExecutor, FConstStructView InOutputDataView)
+{
+	OutputData = InOutputDataView;
+	ExecutorWeak = InExecutor.AsWeak();
+	ProcessOutput(OutputData);
+}
+
+void UDataLinkProcessorBlueprint::OnFinalize(const FDataLinkExecutor& InExecutor, EDataLinkExecutionResult InExecutionResult)
+{
+	ExecutorWeak = InExecutor.AsWeak();
+	Finalize(InExecutionResult);
+}
+
+UWorld* UDataLinkProcessorBlueprint::GetWorld() const
+{
+	if (UObject* ContextObject = GetContextObject())
+	{
+		return ContextObject->GetWorld();
+	}
+	return nullptr;
+}
diff --git a/Engine/Plugins/Experimental/DataLink/Source/DataLink/Private/Nodes/Script/DataLinkScriptNode.cpp b/Engine/Plugins/Experimental/DataLink/Source/DataLink/Private/Nodes/Script/DataLinkScriptNode.cpp
--- a/Engine/Plugins/Experimental/DataLink/Source/DataLink/Private/Nodes/Script/DataLinkScriptNode.cpp
+++ b/Engine/Plugins/Experimental/DataLink/Source/DataLink/Private/Nodes/Script/DataLinkScriptNode.cpp
@@ -14,7 +14,12 @@ void UDataLinkScriptNode::Execute(const UDataLinkNode* InNode, FDataLinkExecutor
 	OnExecute();
 }
 
-bool UDataLinkScriptNode::Succeed(const FInstancedStruct& OutputData)
+void UDataLinkScriptNode::Stop()
+{
+	OnStop();
+}
+
+bool UDataLinkScriptNode::Succeed(const FInstancedStruct& OutputData, bool bPersistExecution)
 {
 	const TSharedPtr<FDataLinkExecutor> Executor = ExecutorWeak.Pin();
 	if (!Executor.IsValid())
@@ -41,7 +46,15 @@ bool UDataLinkScriptNode::Succeed(const FInstancedStruct& OutputData)
 	}
 
 	OutputDataStruct->CopyScriptStruct(OutputDataView.GetMemory(), OutputData.GetMemory());
-	Executor->Next(Node);
+
+	if (bPersistExecution)
+	{
+		Executor->NextPersist(Node);
+	}
+	else
+	{
+		Executor->Next(Node);
+	}
 	return true;
 }
 
@@ -96,6 +109,8 @@ DEFINE_FUNCTION(UDataLinkScriptNode::execSucceedWildcard)
 	const FStructProperty* OutputStructProperty = CastField<FStructProperty>(Stack.MostRecentProperty);
 	void* OutputStructAddress = Stack.MostRecentPropertyAddress;
 
+	P_GET_UBOOL(bPersistExecution);
+
 	P_FINISH;
 
 	*(bool*)RESULT_PARAM = false;
@@ -127,7 +142,16 @@ DEFINE_FUNCTION(UDataLinkScriptNode::execSucceedWildcard)
 		if (OutputDataView.IsValid())
 		{
 			OutputStructProperty->Struct->CopyScriptStruct(OutputDataView.GetMemory(), OutputStructAddress);
-			Executor->Next(P_THIS->Node);
+
+			if (bPersistExecution)
+			{
+				Executor->NextPersist(P_THIS->Node);
+			}
+			else
+			{
+				Executor->Next(P_THIS->Node);
+			}
+
 			*(bool*)RESULT_PARAM = true;
 		}
 		else
diff --git a/Engine/Plugins/Experimental/DataLink/Source/DataLink/Private/Nodes/Script/DataLinkScriptNodeWrapper.cpp b/Engine/Plugins/Experimental/DataLink/Source/DataLink/Private/Nodes/Script/DataLinkScriptNodeWrapper.cpp
--- a/Engine/Plugins/Experimental/DataLink/Source/DataLink/Private/Nodes/Script/DataLinkScriptNodeWrapper.cpp
+++ b/Engine/Plugins/Experimental/DataLink/Source/DataLink/Private/Nodes/Script/DataLinkScriptNodeWrapper.cpp
@@ -75,6 +75,22 @@ EDataLinkExecutionReply UDataLinkScriptNodeWrapper::OnExecute(FDataLinkExecutor&
 	return EDataLinkExecutionReply::Handled;
 }
 
+void UDataLinkScriptNodeWrapper::OnStop(const FDataLinkExecutor& InExecutor) const
+{
+	if (!NodeClass)
+	{
+		return;
+	}
+
+	const FDataLinkNodeInstance& NodeInstance = InExecutor.GetNodeInstance(this);
+
+	const FDataLinkScriptNodeInstance& Instance = NodeInstance.GetInstanceData().Get<const FDataLinkScriptNodeInstance>();
+	if (Instance.Node)
+	{
+		Instance.Node->Stop();
+	}
+}
+
 UDataLinkScriptNode* UDataLinkScriptNodeWrapper::GetDefaultNode() const
 {
 	if (NodeClass)
diff --git a/Engine/Plugins/Experimental/DataLink/Source/DataLink/Public/DataLinkProcessor.h b/Engine/Plugins/Experimental/DataLink/Source/DataLink/Public/DataLinkProcessor.h
--- a/Engine/Plugins/Experimental/DataLink/Source/DataLink/Public/DataLinkProcessor.h
+++ b/Engine/Plugins/Experimental/DataLink/Source/DataLink/Public/DataLinkProcessor.h
@@ -7,26 +7,7 @@
 #include "DataLinkProcessor.generated.h"
 
 class FDataLinkExecutor;
-
-/** Short-lived object wrapping its owning executor for processors to call */
-struct FDataLinkProcessorContext
-{
-	explicit FDataLinkProcessorContext(const FDataLinkExecutor& InExecutor);
-
-	/**
-	 * Gets a brief description of the context for debugging purposes.
-	 * Returns a valid context name when WITH_DATALINK_CONTEXT is true.
-	 * Note: WITH_DATALINK_CONTEXT is true in builds with logging enabled.
-	 */
-	DATALINK_API FStringView GetContextName() const;
-
-	/** Gets the object that is responsible for this execution */
-	DATALINK_API UObject* GetContextObject() const;
-
-private:
-	/** Executor that owns this context */
-	const FDataLinkExecutor& Executor;
-};
+enum class EDataLinkExecutionResult : uint8;
 
 /** Processors execute logic on the output data of a data link graph execution */
 UCLASS(MinimalAPI, Abstract, EditInlineNew)
@@ -35,11 +16,25 @@ class UDataLinkProcessor : public UObject
 	GENERATED_BODY()
 
 public:
+	/** Called at the start of the data link graph execution, prior to running any data link node */
+	void Initialize(const FDataLinkExecutor& InExecutor);
+
 	/** Called when the data link graph execution has output data */
-	void ProcessOutput(FDataLinkProcessorContext& InContext, FConstStructView InOutputDataView);
+	void ProcessOutput(const FDataLinkExecutor& InExecutor, FConstStructView InOutputDataView);
+
+	/** Called at the end of the data link graph execution */
+	void Finalize(const FDataLinkExecutor& InExecutor, EDataLinkExecutionResult InExecutionResult);
 
 protected:
-	virtual void OnProcessOutput(FDataLinkProcessorContext& InContext, FConstStructView InOutputDataView)
+	virtual void OnInitialize(const FDataLinkExecutor& InExecutor)
+	{
+	}
+
+	virtual void OnProcessOutput(const FDataLinkExecutor& InExecutor, FConstStructView InOutputDataView)
+	{
+	}
+
+	virtual void OnFinalize(const FDataLinkExecutor& InExecutor, EDataLinkExecutionResult InExecutionResult)
 	{
 	}
 };
diff --git a/Engine/Plugins/Experimental/DataLink/Source/DataLink/Public/DataLinkProcessorBlueprint.h b/Engine/Plugins/Experimental/DataLink/Source/DataLink/Public/DataLinkProcessorBlueprint.h
new file mode 100644
--- /dev/null
+++ b/Engine/Plugins/Experimental/DataLink/Source/DataLink/Public/DataLinkProcessorBlueprint.h
@@ -0,0 +1,49 @@
+// Copyright Epic Games, Inc. All Rights Reserved.
+
+#pragma once
+
+#include "DataLinkProcessor.h"
+#include "DataLinkProcessorBlueprint.generated.h"
+
+UCLASS(MinimalAPI, Abstract, Blueprintable)
+class UDataLinkProcessorBlueprint : public UDataLinkProcessor
+{
+	GENERATED_BODY()
+
+public:
+	UFUNCTION(BlueprintPure, Category="Data Link")
+	UObject* GetContextObject() const;
+
+	/** Called at the start of the data link graph execution, prior to running any data link node */
+	UFUNCTION(BlueprintImplementableEvent)
+	void Initialize();
+
+	/** Called when the data link graph execution has output data */
+	UFUNCTION(BlueprintImplementableEvent)
+	void ProcessOutput(const FInstancedStruct& InOutputData);
+
+	/** Called at the end of the data link graph execution */
+	UFUNCTION(BlueprintImplementableEvent)
+	void Finalize(EDataLinkExecutionResult InExecutionResult);
+
+protected:
+	//~ Begin UDataLinkProcessor
+	virtual void OnInitialize(const FDataLinkExecutor& InExecutor);
+	virtual void OnProcessOutput(const FDataLinkExecutor& InExecutor, FConstStructView InOutputDataView) override;
+	virtual void OnFinalize(const FDataLinkExecutor& InExecutor, EDataLinkExecutionResult InExecutionResult) override;
+	//~ End UDataLinkProcessor
+
+	//~ Begin UObject
+	virtual UWorld* GetWorld() const override;
+	//~ End UObject
+
+	/**
+	 * Cached output data to avoid memory allocation if the output keeps being the same struct. Only reallocs if struct mismatches.
+	 * Unlike FConstStructView FInstancedStruct is supported in Blueprints.
+	 */
+	UPROPERTY()
+	FInstancedStruct OutputData;
+
+	/** Object responsible for the execution */
+	TWeakPtr<const FDataLinkExecutor> ExecutorWeak;
+};
diff --git a/Engine/Plugins/Experimental/DataLink/Source/DataLink/Public/Nodes/Script/DataLinkScriptNode.h b/Engine/Plugins/Experimental/DataLink/Source/DataLink/Public/Nodes/Script/DataLinkScriptNode.h
--- a/Engine/Plugins/Experimental/DataLink/Source/DataLink/Public/Nodes/Script/DataLinkScriptNode.h
+++ b/Engine/Plugins/Experimental/DataLink/Source/DataLink/Public/Nodes/Script/DataLinkScriptNode.h
@@ -5,6 +5,8 @@
 #include "DataLinkNode.h"
 #include "DataLinkScriptNode.generated.h"
 
+struct FInstancedStruct;
+
 USTRUCT()
 struct FDataLinkScriptPin
 {
@@ -35,16 +37,24 @@ class UDataLinkScriptNode : public UObject
 public:
 	void Execute(const UDataLinkNode* InNode, FDataLinkExecutor& InExecutor);
 
+	void Stop();
+
 	UFUNCTION(BlueprintImplementableEvent)
 	void OnExecute();
 
-	/** Called to move to the next node to execute (or finish if last node) providing the output data as an instanced struct */
+	UFUNCTION(BlueprintImplementableEvent)
+	void OnStop();
+
+	/**
+	 * Called to move to the next node to execute (or finish if last node) providing the output data as an instanced struct
+	 * Set Persist Execution to true so that the graph does not finish, but keeps pushing data.
+	 */
 	UFUNCTION(BlueprintCallable, Category="Data Link", DisplayName="Succeed")
-	bool Succeed(const FInstancedStruct& OutputData);
+	bool Succeed(const FInstancedStruct& OutputData, bool bPersistExecution);
 
 	/** Called to move to the next node to execute (or finish if last node) and providing the output data as a wildcard */
 	UFUNCTION(BlueprintCallable, CustomThunk, Category="Data Link", DisplayName="Succeed (Wildcard)", meta=(CustomStructureParam="OutputData"))
-	bool SucceedWildcard(int32 OutputData);
+	bool SucceedWildcard(int32 OutputData, bool bPersistExecution);
 
 	/** Called to fail execution */
 	UFUNCTION(BlueprintCallable, Category="Data Link")
diff --git a/Engine/Plugins/Experimental/DataLink/Source/DataLink/Public/Nodes/Script/DataLinkScriptNodeWrapper.h b/Engine/Plugins/Experimental/DataLink/Source/DataLink/Public/Nodes/Script/DataLinkScriptNodeWrapper.h
--- a/Engine/Plugins/Experimental/DataLink/Source/DataLink/Public/Nodes/Script/DataLinkScriptNodeWrapper.h
+++ b/Engine/Plugins/Experimental/DataLink/Source/DataLink/Public/Nodes/Script/DataLinkScriptNodeWrapper.h
@@ -6,6 +6,8 @@
 #include "Templates/SubclassOf.h"
 #include "DataLinkScriptNodeWrapper.generated.h"
 
+#define UE_API DATALINK_API
+
 class UDataLinkScriptNode;
 
 /** Wrapper struct for the Node Instance containing the actual blueprint implementation of the data link logic */
@@ -30,22 +32,23 @@ class UDataLinkScriptNodeWrapper : public UDataLinkNode
 	GENERATED_BODY()
 
 public:
-	DATALINK_API UDataLinkScriptNodeWrapper();
+	UE_API UDataLinkScriptNodeWrapper();
 
 	static FName GetNodeClassPropertyName()
 	{
 		return GET_MEMBER_NAME_CHECKED(UDataLinkScriptNodeWrapper, NodeClass);
 	}
 
-	DATALINK_API void SetNodeClass(TSubclassOf<UDataLinkScriptNode> InNodeClass);
+	UE_API void SetNodeClass(TSubclassOf<UDataLinkScriptNode> InNodeClass);
 
 protected:
 	//~ Begin UDataLinkNode
 #if WITH_EDITOR
-	DATALINK_API virtual void OnBuildMetadata(FDataLinkNodeMetadata& Metadata) const override;
+	UE_API virtual void OnBuildMetadata(FDataLinkNodeMetadata& Metadata) const override;
 #endif
-	DATALINK_API virtual void OnBuildPins(FDataLinkPinBuilder& Inputs, FDataLinkPinBuilder& Outputs) const override;
-	DATALINK_API virtual EDataLinkExecutionReply OnExecute(FDataLinkExecutor& InExecutor) const override;
+	UE_API virtual void OnBuildPins(FDataLinkPinBuilder& Inputs, FDataLinkPinBuilder& Outputs) const override;
+	UE_API virtual EDataLinkExecutionReply OnExecute(FDataLinkExecutor& InExecutor) const override;
+	UE_API virtual void OnStop(const FDataLinkExecutor& InExecutor) const override;
 	//~ End UDataLinkNode
 
 private:
@@ -62,3 +65,5 @@ private:
 	UPROPERTY(VisibleAnywhere, Instanced, Category="Data Link")
 	TObjectPtr<UDataLinkScriptNode> TemplateNode;
 };
+
+#undef UE_API
diff --git a/Engine/Plugins/Experimental/DataLink/Source/DataLinkEditor/Private/Preview/DataLinkPreviewTool.cpp b/Engine/Plugins/Experimental/DataLink/Source/DataLinkEditor/Private/Preview/DataLinkPreviewTool.cpp
--- a/Engine/Plugins/Experimental/DataLink/Source/DataLinkEditor/Private/Preview/DataLinkPreviewTool.cpp
+++ b/Engine/Plugins/Experimental/DataLink/Source/DataLinkEditor/Private/Preview/DataLinkPreviewTool.cpp
@@ -110,7 +110,8 @@ void FDataLinkPreviewTool::RunPreview()
 #if WITH_DATALINK_CONTEXT
 		.SetContextName(TEXT("Graph Editor Preview"))
 #endif
-		.SetContextObject(AssetEditor)
+		// Blueprint nodes could use nodes like 'Delay' that require a world context object. Ensure preview can run these nodes by setting the context to GWorld
+		.SetContextObject(GWorld)
 		.SetSink(Sink)
 		.SetOnOutputData(FOnDataLinkOutputData::CreateSP(this, &FDataLinkPreviewTool::OnPreviewOutputData)));
 
