Motion Design Data Link: added option in md data link actor to execute data link on begin play

diff --git a/Engine/Plugins/Experimental/AvalancheDataLink/Source/AvalancheDataLink/Private/AvaDataLinkActor.cpp b/Engine/Plugins/Experimental/AvalancheDataLink/Source/AvalancheDataLink/Private/AvaDataLinkActor.cpp
--- a/Engine/Plugins/Experimental/AvalancheDataLink/Source/AvalancheDataLink/Private/AvaDataLinkActor.cpp
+++ b/Engine/Plugins/Experimental/AvalancheDataLink/Source/AvalancheDataLink/Private/AvaDataLinkActor.cpp
@@ -25,6 +25,16 @@ void AAvaDataLinkActor::StopDataLinkInstances()
 	}
 }
 
+void AAvaDataLinkActor::BeginPlay()
+{
+	Super::BeginPlay();
+
+	if (bExecuteOnBeginPlay)
+	{
+		ExecuteDataLinkInstances();
+	}
+}
+
 #if WITH_EDITOR
 FString AAvaDataLinkActor::GetDefaultActorLabel() const
 {
diff --git a/Engine/Plugins/Experimental/AvalancheDataLink/Source/AvalancheDataLink/Private/AvaDataLinkActor.h b/Engine/Plugins/Experimental/AvalancheDataLink/Source/AvalancheDataLink/Private/AvaDataLinkActor.h
--- a/Engine/Plugins/Experimental/AvalancheDataLink/Source/AvalancheDataLink/Private/AvaDataLinkActor.h
+++ b/Engine/Plugins/Experimental/AvalancheDataLink/Source/AvalancheDataLink/Private/AvaDataLinkActor.h
@@ -20,6 +20,7 @@ public:
 	void StopDataLinkInstances();
 
 	//~ Begin AActor
+	virtual void BeginPlay() override;
 #if WITH_EDITOR
 	virtual FString GetDefaultActorLabel() const override;
 #endif
@@ -27,4 +28,8 @@ public:
 
 	UPROPERTY(EditAnywhere, Instanced, Category="Motion Design Data Link")
 	TArray<TObjectPtr<UAvaDataLinkInstance>> DataLinkInstances;
+
+	/** Whether to automatically execute the data link instances on Begin Play */
+	UPROPERTY(EditAnywhere, Category="Motion Design Data Link")
+	bool bExecuteOnBeginPlay = false;
 };
