Motion Design: data link instances are stopped on world cleanup + begindestroy on data link actor

diff --git a/Engine/Plugins/Experimental/AvalancheDataLink/Source/AvalancheDataLink/Private/AvaDataLinkActor.cpp b/Engine/Plugins/Experimental/AvalancheDataLink/Source/AvalancheDataLink/Private/AvaDataLinkActor.cpp
--- a/Engine/Plugins/Experimental/AvalancheDataLink/Source/AvalancheDataLink/Private/AvaDataLinkActor.cpp
+++ b/Engine/Plugins/Experimental/AvalancheDataLink/Source/AvalancheDataLink/Private/AvaDataLinkActor.cpp
@@ -2,6 +2,12 @@
 
 #include "AvaDataLinkActor.h"
 #include "AvaDataLinkInstance.h"
+#include "Engine/World.h"
+
+AAvaDataLinkActor::AAvaDataLinkActor()
+{
+	FWorldDelegates::OnWorldCleanup.AddUObject(this, &AAvaDataLinkActor::OnWorldCleanup);
+}
 
 void AAvaDataLinkActor::ExecuteDataLinkInstances()
 {
@@ -41,3 +47,17 @@ FString AAvaDataLinkActor::GetDefaultActorLabel() const
 	return TEXT("Motion Design Data Link Actor");
 }
 #endif
+
+void AAvaDataLinkActor::BeginDestroy()
+{
+	Super::BeginDestroy();
+	StopDataLinkInstances();
+}
+
+void AAvaDataLinkActor::OnWorldCleanup(UWorld* InWorld, bool bInSessionEnded, bool bInCleanupResources)
+{
+	if (InWorld == GetWorld())
+	{
+		StopDataLinkInstances();
+	}
+}
diff --git a/Engine/Plugins/Experimental/AvalancheDataLink/Source/AvalancheDataLink/Private/AvaDataLinkActor.h b/Engine/Plugins/Experimental/AvalancheDataLink/Source/AvalancheDataLink/Private/AvaDataLinkActor.h
--- a/Engine/Plugins/Experimental/AvalancheDataLink/Source/AvalancheDataLink/Private/AvaDataLinkActor.h
+++ b/Engine/Plugins/Experimental/AvalancheDataLink/Source/AvalancheDataLink/Private/AvaDataLinkActor.h
@@ -13,6 +13,8 @@ class AAvaDataLinkActor : public AActor
 	GENERATED_BODY()
 
 public:
+	AAvaDataLinkActor();
+
 	UFUNCTION(BlueprintCallable, CallInEditor, Category="Motion Design Data Link")
 	void ExecuteDataLinkInstances();
 
@@ -26,6 +28,12 @@ public:
 #endif
 	//~ End AActor
 
+	//~ Begin UObject
+	virtual void BeginDestroy() override;
+	//~ End UObject
+
+	void OnWorldCleanup(UWorld* InWorld, bool bInSessionEnded, bool bInCleanupResources);
+
 	UPROPERTY(EditAnywhere, Instanced, Category="Motion Design Data Link")
 	TArray<TObjectPtr<UAvaDataLinkInstance>> DataLinkInstances;
 
