Motion Design: Added data link and ticker actor to the motion design toolbox

diff --git a/Engine/Plugins/Experimental/AvalancheDataLink/Source/AvalancheDataLink/Private/AvaDataLinkActor.h b/Engine/Plugins/Experimental/AvalancheDataLink/Source/AvalancheDataLink/Private/AvaDataLinkActor.h
deleted file mode 100644
diff --git a/Engine/Plugins/Experimental/AvalancheDataLink/Source/AvalancheDataLink/Public/AvaDataLinkActor.h b/Engine/Plugins/Experimental/AvalancheDataLink/Source/AvalancheDataLink/Public/AvaDataLinkActor.h
new file mode 100644
--- /dev/null
+++ b/Engine/Plugins/Experimental/AvalancheDataLink/Source/AvalancheDataLink/Public/AvaDataLinkActor.h
@@ -0,0 +1,43 @@
+// Copyright Epic Games, Inc. All Rights Reserved.
+
+#pragma once
+
+#include "GameFramework/Actor.h"
+#include "AvaDataLinkActor.generated.h"
+
+class UAvaDataLinkInstance;
+
+UCLASS(MinimalAPI, DisplayName="Motion Design Data Link Actor", HideCategories=(Activation, Actor, AssetUserData, Collision, Cooking, DataLayers, HLOD, Input, LevelInstance, Mobility, Navigation, Networking, Physics, Rendering, Replication, Tags, Transform, WorldPartition))
+class AAvaDataLinkActor : public AActor
+{
+	GENERATED_BODY()
+
+public:
+	AAvaDataLinkActor();
+
+	UFUNCTION(BlueprintCallable, CallInEditor, Category="Motion Design Data Link")
+	void ExecuteDataLinkInstances();
+
+	UFUNCTION(BlueprintCallable, CallInEditor, Category="Motion Design Data Link")
+	void StopDataLinkInstances();
+
+	//~ Begin AActor
+	virtual void BeginPlay() override;
+#if WITH_EDITOR
+	virtual FString GetDefaultActorLabel() const override;
+#endif
+	//~ End AActor
+
+	//~ Begin UObject
+	virtual void BeginDestroy() override;
+	//~ End UObject
+
+	void OnWorldCleanup(UWorld* InWorld, bool bInSessionEnded, bool bInCleanupResources);
+
+	UPROPERTY(EditAnywhere, Instanced, Category="Motion Design Data Link")
+	TArray<TObjectPtr<UAvaDataLinkInstance>> DataLinkInstances;
+
+	/** Whether to automatically execute the data link instances on Begin Play */
+	UPROPERTY(EditAnywhere, Category="Motion Design Data Link")
+	bool bExecuteOnBeginPlay = false;
+};
diff --git a/Engine/Plugins/Experimental/AvalancheDataLink/Source/AvalancheDataLinkEditor/AvalancheDataLinkEditor.Build.cs b/Engine/Plugins/Experimental/AvalancheDataLink/Source/AvalancheDataLinkEditor/AvalancheDataLinkEditor.Build.cs
--- a/Engine/Plugins/Experimental/AvalancheDataLink/Source/AvalancheDataLinkEditor/AvalancheDataLinkEditor.Build.cs
+++ b/Engine/Plugins/Experimental/AvalancheDataLink/Source/AvalancheDataLinkEditor/AvalancheDataLinkEditor.Build.cs
@@ -12,12 +12,15 @@ public class AvalancheDataLinkEditor : ModuleRules
             new string[]
             {
                 "AvalancheDataLink",
+                "AvalancheInteractiveTools",
                 "Core",
                 "CoreUObject",
                 "Engine",
+                "InteractiveToolsFramework",
                 "PropertyEditor",
                 "Slate",
-                "SlateCore", 
+                "SlateCore",
+                "UnrealEd",
             }
         );
     }
diff --git a/Engine/Plugins/Experimental/AvalancheDataLink/Source/AvalancheDataLinkEditor/Private/AvaDataLinkActorTool.cpp b/Engine/Plugins/Experimental/AvalancheDataLink/Source/AvalancheDataLinkEditor/Private/AvaDataLinkActorTool.cpp
new file mode 100644
--- /dev/null
+++ b/Engine/Plugins/Experimental/AvalancheDataLink/Source/AvalancheDataLinkEditor/Private/AvaDataLinkActorTool.cpp
@@ -0,0 +1,47 @@
+// Copyright Epic Games, Inc. All Rights Reserved.
+
+#include "AvaDataLinkActorTool.h"
+#include "AvaDataLinkActor.h"
+#include "AvaDataLinkEditorCommands.h"
+#include "Builders/AvaInteractiveToolsToolBuilder.h"
+#include "Styling/SlateIconFinder.h"
+
+UAvaDataLinkActorTool::UAvaDataLinkActorTool()
+{
+	ActorClass = AAvaDataLinkActor::StaticClass();
+}
+
+bool UAvaDataLinkActorTool::ShouldForceDefaultAction() const
+{
+	// Force to always use default action since Data Link Actor has no transform
+	return true;
+}
+
+bool UAvaDataLinkActorTool::SupportsDefaultAction() const
+{
+	return true;
+}
+
+bool UAvaDataLinkActorTool::OnBegin()
+{
+	// Overriden since default OnBegin returns false on if ViewportPlannerClass is null
+	BeginTransaction();
+	return true;
+}
+
+void UAvaDataLinkActorTool::OnRegisterTool(IAvalancheInteractiveToolsModule* InAITModule)
+{
+	Super::OnRegisterTool(InAITModule);
+
+	FAvaInteractiveToolsToolParameters ToolParameters;
+	ToolParameters.UICommand = UE::AvaDataLink::FEditorCommands::Get().DataLinkActorTool;
+	ToolParameters.ToolIdentifier = TEXT("Motion Design Data Link Actor Tool");
+	ToolParameters.FactoryClass = ActorClass;
+	ToolParameters.CreateBuilder.BindLambda(
+		[](UEdMode* InEdMode)
+		{
+			return UAvaInteractiveToolsToolBuilder::CreateToolBuilder<UAvaDataLinkActorTool>(InEdMode);
+		});
+
+	InAITModule->RegisterTool(IAvalancheInteractiveToolsModule::CategoryNameActor, MoveTemp(ToolParameters));
+}
diff --git a/Engine/Plugins/Experimental/AvalancheDataLink/Source/AvalancheDataLinkEditor/Private/AvaDataLinkActorTool.h b/Engine/Plugins/Experimental/AvalancheDataLink/Source/AvalancheDataLinkEditor/Private/AvaDataLinkActorTool.h
new file mode 100644
--- /dev/null
+++ b/Engine/Plugins/Experimental/AvalancheDataLink/Source/AvalancheDataLinkEditor/Private/AvaDataLinkActorTool.h
@@ -0,0 +1,22 @@
+// Copyright Epic Games, Inc. All Rights Reserved.
+
+#pragma once
+
+#include "Tools/AvaInteractiveToolsActorToolBase.h"
+#include "AvaDataLinkActorTool.generated.h"
+
+UCLASS()
+class UAvaDataLinkActorTool : public UAvaInteractiveToolsActorToolBase
+{
+	GENERATED_BODY()
+
+public:
+	UAvaDataLinkActorTool();
+
+	//~ Begin UAvaInteractiveToolsToolBase
+	virtual bool ShouldForceDefaultAction() const override;
+	virtual bool SupportsDefaultAction() const override;
+	virtual bool OnBegin() override;
+	virtual void OnRegisterTool(IAvalancheInteractiveToolsModule* InAITModule) override;
+	//~ End UAvaInteractiveToolsToolBase
+};
diff --git a/Engine/Plugins/Experimental/AvalancheDataLink/Source/AvalancheDataLinkEditor/Private/AvaDataLinkEditorCommands.cpp b/Engine/Plugins/Experimental/AvalancheDataLink/Source/AvalancheDataLinkEditor/Private/AvaDataLinkEditorCommands.cpp
new file mode 100644
--- /dev/null
+++ b/Engine/Plugins/Experimental/AvalancheDataLink/Source/AvalancheDataLinkEditor/Private/AvaDataLinkEditorCommands.cpp
@@ -0,0 +1,34 @@
+// Copyright Epic Games, Inc. All Rights Reserved.
+
+#include "AvaDataLinkEditorCommands.h"
+#include "Styling/AppStyle.h"
+
+#define LOCTEXT_NAMESPACE "AvaDataLinkEditorCommands"
+
+namespace UE::AvaDataLink
+{
+
+FEditorCommands::FEditorCommands()
+	: TCommands(TEXT("AvaDataLinkEditor")
+	, LOCTEXT("AvaDataLinkEditor", "Motion Design Data Link Editor")
+	, NAME_None
+	, FAppStyle::GetAppStyleSetName())
+{
+}
+
+void FEditorCommands::RegisterCommands()
+{
+	FUICommandInfo::MakeCommandInfo(SharedThis(this)
+		, DataLinkActorTool
+		, TEXT("DataLinkActorTool")
+		, LOCTEXT("DataLinkActorTool", "Data Link Actor")
+		, LOCTEXT("DataLinkActorTool_ToolTip", "Create a Motion Design Data Link Actor")
+		, FSlateIcon(FAppStyle::GetAppStyleSetName(), "ClassIcon.Actor")
+		, EUserInterfaceActionType::ToggleButton
+		, FInputChord()
+	);
+}
+
+} // UE::AvaDataLink
+
+#undef LOCTEXT_NAMESPACE
diff --git a/Engine/Plugins/Experimental/AvalancheDataLink/Source/AvalancheDataLinkEditor/Private/AvaDataLinkEditorCommands.h b/Engine/Plugins/Experimental/AvalancheDataLink/Source/AvalancheDataLinkEditor/Private/AvaDataLinkEditorCommands.h
new file mode 100644
--- /dev/null
+++ b/Engine/Plugins/Experimental/AvalancheDataLink/Source/AvalancheDataLinkEditor/Private/AvaDataLinkEditorCommands.h
@@ -0,0 +1,23 @@
+// Copyright Epic Games, Inc. All Rights Reserved.
+
+#pragma once
+
+#include "Framework/Commands/Commands.h"
+
+namespace UE::AvaDataLink
+{
+
+class FEditorCommands : public TCommands<FEditorCommands>
+{
+public:
+	FEditorCommands();
+
+	//~ Begin TCommands
+	virtual void RegisterCommands() override;
+	//~ End TCommands
+
+	/** Command used to spawn the Motion Design Data Link Actor */
+	TSharedPtr<FUICommandInfo> DataLinkActorTool;
+};
+
+} // UE::AvaDataLink
diff --git a/Engine/Plugins/Experimental/AvalancheDataLink/Source/AvalancheDataLinkEditor/Private/AvaDataLinkEditorModule.cpp b/Engine/Plugins/Experimental/AvalancheDataLink/Source/AvalancheDataLinkEditor/Private/AvaDataLinkEditorModule.cpp
--- a/Engine/Plugins/Experimental/AvalancheDataLink/Source/AvalancheDataLinkEditor/Private/AvaDataLinkEditorModule.cpp
+++ b/Engine/Plugins/Experimental/AvalancheDataLink/Source/AvalancheDataLinkEditor/Private/AvaDataLinkEditorModule.cpp
@@ -1,6 +1,7 @@
 // Copyright Epic Games, Inc. All Rights Reserved.
 
 #include "AvaDataLinkEditorModule.h"
+#include "AvaDataLinkEditorCommands.h"
 #include "AvaDataLinkInstance.h"
 #include "DetailsView/AvaDataLinkInstanceCustomization.h"
 #include "Modules/ModuleManager.h"
@@ -10,6 +11,8 @@ IMPLEMENT_MODULE(FAvaDataLinkEditorModule, AvalancheDataLinkEditor)
 
 void FAvaDataLinkEditorModule::StartupModule()
 {
+	UE::AvaDataLink::FEditorCommands::Register();
+
 	FPropertyEditorModule& PropertyEditorModule = FModuleManager::LoadModuleChecked<FPropertyEditorModule>("PropertyEditor");
 
 	PropertyEditorModule.RegisterCustomClassLayout(CustomizedClasses.Add_GetRef(UAvaDataLinkInstance::StaticClass()->GetFName())
@@ -18,6 +21,8 @@ void FAvaDataLinkEditorModule::StartupModule()
 
 void FAvaDataLinkEditorModule::ShutdownModule()
 {
+	UE::AvaDataLink::FEditorCommands::Unregister();
+
 	if (FPropertyEditorModule* PropertyEditorModule = FModuleManager::GetModulePtr<FPropertyEditorModule>("PropertyEditor"))
 	{
 		for (FName CustomizedClass : CustomizedClasses)
diff --git a/Engine/Plugins/VirtualProduction/Avalanche/Source/AvalancheEditor/Private/AvaEditorCommands.cpp b/Engine/Plugins/VirtualProduction/Avalanche/Source/AvalancheEditor/Private/AvaEditorCommands.cpp
--- a/Engine/Plugins/VirtualProduction/Avalanche/Source/AvalancheEditor/Private/AvaEditorCommands.cpp
+++ b/Engine/Plugins/VirtualProduction/Avalanche/Source/AvalancheEditor/Private/AvaEditorCommands.cpp
@@ -374,6 +374,16 @@ void FAvaEditorCommands::RegisterToolsCommands()
 		, EUserInterfaceActionType::ToggleButton
 		, FInputChord());
 
+	FUICommandInfo::MakeCommandInfo(AsShared()
+		, TickerTool
+		, TEXT("TickerTool")
+		, LOCTEXT("TickerTool", "Ticker")
+		, LOCTEXT("TickerTool_Tooltip", "Create a Motion Design Ticker actor in the viewport.")
+		, FSlateIconFinder::FindIconForClass(AActor::StaticClass())
+		, EUserInterfaceActionType::ToggleButton
+		, FInputChord()
+	);
+
 	UI_COMMAND(PointLightTool
 		, "Point Light"
 		, "Create a point light actor in the viewport."
diff --git a/Engine/Plugins/VirtualProduction/Avalanche/Source/AvalancheEditor/Private/AvaEditorCommands.h b/Engine/Plugins/VirtualProduction/Avalanche/Source/AvalancheEditor/Private/AvaEditorCommands.h
--- a/Engine/Plugins/VirtualProduction/Avalanche/Source/AvalancheEditor/Private/AvaEditorCommands.h
+++ b/Engine/Plugins/VirtualProduction/Avalanche/Source/AvalancheEditor/Private/AvaEditorCommands.h
@@ -85,6 +85,7 @@ public:
 	TSharedPtr<FUICommandInfo> CameraRigRailTool;
 	TSharedPtr<FUICommandInfo> CameraShakeSourceTool;
 	TSharedPtr<FUICommandInfo> AvaPostProcessVolumeTool;
+	TSharedPtr<FUICommandInfo> TickerTool;
 
 	TSharedPtr<FUICommandInfo> PointLightTool;
 	TSharedPtr<FUICommandInfo> DirectionalLightTool;
diff --git a/Engine/Plugins/VirtualProduction/Avalanche/Source/AvalancheEditor/Private/Ticker/AvaTickerActorTool.cpp b/Engine/Plugins/VirtualProduction/Avalanche/Source/AvalancheEditor/Private/Ticker/AvaTickerActorTool.cpp
new file mode 100644
--- /dev/null
+++ b/Engine/Plugins/VirtualProduction/Avalanche/Source/AvalancheEditor/Private/Ticker/AvaTickerActorTool.cpp
@@ -0,0 +1,28 @@
+// Copyright Epic Games, Inc. All Rights Reserved.
+
+#include "AvaTickerActorTool.h"
+#include "AvaEditorCommands.h"
+#include "Builders/AvaInteractiveToolsToolBuilder.h"
+#include "Framework/Ticker/AvaTickerActor.h"
+
+UAvaTickerActorTool::UAvaTickerActorTool()
+{
+	ActorClass = AAvaTickerActor::StaticClass();
+}
+
+void UAvaTickerActorTool::OnRegisterTool(IAvalancheInteractiveToolsModule* InAITModule)
+{
+	Super::OnRegisterTool(InAITModule);
+
+	FAvaInteractiveToolsToolParameters ToolParameters;
+	ToolParameters.UICommand = FAvaEditorCommands::Get().TickerTool;
+	ToolParameters.ToolIdentifier = TEXT("Motion Design Ticker Tool");
+	ToolParameters.FactoryClass = ActorClass;
+	ToolParameters.CreateBuilder.BindLambda(
+		[](UEdMode* InEdMode)
+		{
+			return UAvaInteractiveToolsToolBuilder::CreateToolBuilder<UAvaTickerActorTool>(InEdMode);
+		});
+
+	InAITModule->RegisterTool(IAvalancheInteractiveToolsModule::CategoryNameActor, MoveTemp(ToolParameters));
+}
diff --git a/Engine/Plugins/VirtualProduction/Avalanche/Source/AvalancheEditor/Private/Ticker/AvaTickerActorTool.h b/Engine/Plugins/VirtualProduction/Avalanche/Source/AvalancheEditor/Private/Ticker/AvaTickerActorTool.h
new file mode 100644
--- /dev/null
+++ b/Engine/Plugins/VirtualProduction/Avalanche/Source/AvalancheEditor/Private/Ticker/AvaTickerActorTool.h
@@ -0,0 +1,19 @@
+// Copyright Epic Games, Inc. All Rights Reserved.
+
+#pragma once
+
+#include "Tools/AvaInteractiveToolsActorPointToolBase.h"
+#include "AvaTickerActorTool.generated.h"
+
+UCLASS()
+class UAvaTickerActorTool : public UAvaInteractiveToolsActorPointToolBase
+{
+	GENERATED_BODY()
+
+public:
+	UAvaTickerActorTool();
+
+	//~ Begin UAvaInteractiveToolsToolBase
+	virtual void OnRegisterTool(IAvalancheInteractiveToolsModule* InAITModule) override;
+	//~ End UAvaInteractiveToolsToolBase
+};
diff --git a/Engine/Plugins/VirtualProduction/Avalanche/Source/AvalancheInteractiveTools/Private/Tools/AvaInteractiveToolsToolBase.cpp b/Engine/Plugins/VirtualProduction/Avalanche/Source/AvalancheInteractiveTools/Private/Tools/AvaInteractiveToolsToolBase.cpp
--- a/Engine/Plugins/VirtualProduction/Avalanche/Source/AvalancheInteractiveTools/Private/Tools/AvaInteractiveToolsToolBase.cpp
+++ b/Engine/Plugins/VirtualProduction/Avalanche/Source/AvalancheInteractiveTools/Private/Tools/AvaInteractiveToolsToolBase.cpp
@@ -463,6 +463,26 @@ void UAvaInteractiveToolsToolBase::OnViewportPlannerComplete()
 	RequestShutdown(EToolShutdownType::Completed);
 }
 
+void UAvaInteractiveToolsToolBase::OnActorSpawned(AActor* InActor)
+{
+	if (!InActor || InActor->bIsEditorPreviewActor)
+	{
+		return;
+	}
+
+	// Delay to select actor after setup completed
+	FTSTicker::GetCoreTicker().AddTicker(FTickerDelegate::CreateWeakLambda(InActor,
+		[InActor](float InDelta)
+		{
+			if (GEditor)
+			{
+				GEditor->SelectNone(true, true);
+				GEditor->SelectActor(InActor, true, true, false, true);
+			}
+			return false;
+		}));
+}
+
 FInputRayHit UAvaInteractiveToolsToolBase::CanBeginSingleClickAndDragSequence(const FInputDeviceRay& InPressPos)
 {
 	// Always hits every place in the viewport
@@ -580,7 +600,7 @@ void UAvaInteractiveToolsToolBase::Activate(const FString& InToolIdentifier, boo
 	{
 		if (bForceDefaultAction)
 		{
-			UE_LOG(LogAvaInteractiveTools, Warning, TEXT("Alt used to force a tool into using the default action, but it does not support the default action."));
+			UE_LOG(LogAvaInteractiveTools, Warning, TEXT("Tool forced into using the default action, but it does not support default action."));
 		}
 
 		RequestShutdown(EToolShutdownType::Cancel);
diff --git a/Engine/Plugins/VirtualProduction/Avalanche/Source/AvalancheInteractiveTools/Public/Tools/AvaInteractiveToolsToolBase.h b/Engine/Plugins/VirtualProduction/Avalanche/Source/AvalancheInteractiveTools/Public/Tools/AvaInteractiveToolsToolBase.h
--- a/Engine/Plugins/VirtualProduction/Avalanche/Source/AvalancheInteractiveTools/Public/Tools/AvaInteractiveToolsToolBase.h
+++ b/Engine/Plugins/VirtualProduction/Avalanche/Source/AvalancheInteractiveTools/Public/Tools/AvaInteractiveToolsToolBase.h
@@ -91,7 +91,7 @@ public:
 
 	virtual void OnViewportPlannerUpdate() {}
 	virtual void OnViewportPlannerComplete();
-	virtual void OnActorSpawned(AActor* InActor) {}
+	virtual void OnActorSpawned(AActor* InActor);
 
 	/** Calls the below function with the default camera distance. */
 	bool ViewportPositionToWorldPositionAndOrientation(EAvaViewportStatus InViewportStatus, const FVector2f& InViewportPosition,
@@ -192,5 +192,5 @@ protected:
 
 	bool IsMotionDesignViewport() const;
 
-	bool ShouldForceDefaultAction() const;
+	virtual bool ShouldForceDefaultAction() const;
 };
diff --git a/Engine/Plugins/VirtualProduction/Avalanche/Source/AvalancheShapesEditor/Private/Tools/AvaShapesEditorShapeAreaToolBase.cpp b/Engine/Plugins/VirtualProduction/Avalanche/Source/AvalancheShapesEditor/Private/Tools/AvaShapesEditorShapeAreaToolBase.cpp
--- a/Engine/Plugins/VirtualProduction/Avalanche/Source/AvalancheShapesEditor/Private/Tools/AvaShapesEditorShapeAreaToolBase.cpp
+++ b/Engine/Plugins/VirtualProduction/Avalanche/Source/AvalancheShapesEditor/Private/Tools/AvaShapesEditorShapeAreaToolBase.cpp
@@ -76,18 +76,9 @@ void UAvaShapesEditorShapeAreaToolBase::OnViewportPlannerComplete()
 		SpawnedActor = SpawnActor(ActorClass, EAvaViewportStatus::Focused, AreaPlanner->GetCenterPosition(), /* Preview */ false);
 	}
 
-	SelectActor(SpawnedActor);
-
 	Super::OnViewportPlannerComplete();
 }
 
-void UAvaShapesEditorShapeAreaToolBase::DefaultAction()
-{
-	Super::DefaultAction();
-
-	SelectActor(SpawnedActor);
-}
-
 void UAvaShapesEditorShapeAreaToolBase::UpdateShapeSize(AAvaShapeActor* InShapeActor) const
 {
 	if (!InShapeActor)
@@ -103,21 +94,3 @@ void UAvaShapesEditorShapeAreaToolBase::UpdateShapeSize(AAvaShapeActor* InShapeA
 		}
 	}
 }
-
-void UAvaShapesEditorShapeAreaToolBase::SelectActor(AActor* InActor) const
-{
-	if (InActor)
-	{
-		// Delay to select shape after setup completed
-		FTSTicker::GetCoreTicker().AddTicker(FTickerDelegate::CreateWeakLambda(InActor, [InActor](float InDelta)
-		{
-			if (GEditor && InActor)
-			{
-				GEditor->SelectNone(true, true);
-				GEditor->SelectActor(InActor, true, true, false, true);
-			}
-
-			return false;
-		}));
-	}
-}
diff --git a/Engine/Plugins/VirtualProduction/Avalanche/Source/AvalancheShapesEditor/Private/Tools/AvaShapesEditorShapeAreaToolBase.h b/Engine/Plugins/VirtualProduction/Avalanche/Source/AvalancheShapesEditor/Private/Tools/AvaShapesEditorShapeAreaToolBase.h
--- a/Engine/Plugins/VirtualProduction/Avalanche/Source/AvalancheShapesEditor/Private/Tools/AvaShapesEditorShapeAreaToolBase.h
+++ b/Engine/Plugins/VirtualProduction/Avalanche/Source/AvalancheShapesEditor/Private/Tools/AvaShapesEditorShapeAreaToolBase.h
@@ -21,12 +21,10 @@ public:
 	//~ Begin UAvaInteractiveToolsToolBase
 	virtual void OnViewportPlannerUpdate() override;
 	virtual void OnViewportPlannerComplete() override;
-	virtual void DefaultAction() override;
 	//~ End UAvaInteractiveToolsToolBase
 
 protected:
 	void UpdateShapeSize(AAvaShapeActor* InShapeActor) const;
-	void SelectActor(AActor* InActor) const;
 
 	//~ Begin UAvaInteractiveToolsToolBase
 	virtual AActor* SpawnActor(TSubclassOf<AActor> InActorClass, EAvaViewportStatus InViewportStatus, const FVector2f& InViewportPosition, bool bInPreview, FString* InActorLabelOverride = nullptr) override;
